GeometryProcessing: move all public types in GeometryProcessing plugin into UE::Geometry:: namespace, as well as some code in MeshModelingToolset. Fix fallout in Water, Paint Mode, ChaosEditor, DataPrep, Text3D, Unit Tests, internal projects.

diff --git a/Engine/Plugins/Experimental/ChaosEditor/Source/FractureEditor/Private/FractureToolUV.cpp b/Engine/Plugins/Experimental/ChaosEditor/Source/FractureEditor/Private/FractureToolUV.cpp
--- a/Engine/Plugins/Experimental/ChaosEditor/Source/FractureEditor/Private/FractureToolUV.cpp
+++ b/Engine/Plugins/Experimental/ChaosEditor/Source/FractureEditor/Private/FractureToolUV.cpp
@@ -21,6 +21,7 @@
 
 #define LOCTEXT_NAMESPACE "FractureToolAutoUV"
 
+using namespace UE::Geometry;
 
 UFractureToolAutoUV::UFractureToolAutoUV(const FObjectInitializer& ObjInit)
 	: Super(ObjInit)
@@ -217,9 +218,9 @@ int32 UFractureToolAutoUV::ExecuteFracture(const FFractureToolContext& FractureC
 		FImageDimensions Dimensions(OutputRes, OutputRes);
 		TextureBuilder.Initialize(FTexture2DBuilder::ETextureType::Color, Dimensions);
 
-		TImageBuilder<FVector3f> ImageBuilder;
+		UE::Geometry::TImageBuilder<UE::Geometry::FVector3f> ImageBuilder;
 		ImageBuilder.SetDimensions(Dimensions);
-		ImageBuilder.Clear(FVector3f(0, 0, 0));
+		ImageBuilder.Clear(UE::Geometry::FVector3f(0, 0, 0));
 
 		UE::PlanarCut::TextureInternalSurfaces(Collection, AutoUVSettings->MaxDistance, FMath::CeilToInt(AutoUVSettings->GutterSize), ImageBuilder);
 
diff --git a/Engine/Plugins/Experimental/Enterprise/DataprepGeometryOperations/Source/Private/DataprepGeometryOperations.cpp b/Engine/Plugins/Experimental/Enterprise/DataprepGeometryOperations/Source/Private/DataprepGeometryOperations.cpp
--- a/Engine/Plugins/Experimental/Enterprise/DataprepGeometryOperations/Source/Private/DataprepGeometryOperations.cpp
+++ b/Engine/Plugins/Experimental/Enterprise/DataprepGeometryOperations/Source/Private/DataprepGeometryOperations.cpp
@@ -79,11 +79,11 @@ void UDataprepRemeshOperation::OnExecution_Implementation(const FDataprepContext
 		
 		FMeshDescription* MeshDescription = StaticMesh->GetMeshDescription(0);
 
-		TSharedPtr<FDynamicMesh3, ESPMode::ThreadSafe> OriginalMesh = MakeShared<FDynamicMesh3, ESPMode::ThreadSafe>();
+		TSharedPtr<UE::Geometry::FDynamicMesh3, ESPMode::ThreadSafe> OriginalMesh = MakeShared<UE::Geometry::FDynamicMesh3, ESPMode::ThreadSafe>();
 		FMeshDescriptionToDynamicMesh MeshDescriptionToDynamicMesh;
 		MeshDescriptionToDynamicMesh.Convert(MeshDescription, *OriginalMesh);
 
-		FRemeshMeshOp Op;
+		UE::Geometry::FRemeshMeshOp Op;
 
 		Op.RemeshType = RemeshType;
 		Op.bCollapses = true;
@@ -100,7 +100,7 @@ void UDataprepRemeshOperation::OnExecution_Implementation(const FDataprepContext
 		Op.SmoothingStrength = SmoothingStrength;
 		Op.SmoothingType = ERemeshSmoothingType::MeanValue;
 
-		TSharedPtr<FDynamicMeshAABBTree3, ESPMode::ThreadSafe> OriginalMeshSpatial = MakeShared<FDynamicMeshAABBTree3, ESPMode::ThreadSafe>(OriginalMesh.Get(), true);
+		TSharedPtr<UE::Geometry::FDynamicMeshAABBTree3, ESPMode::ThreadSafe> OriginalMeshSpatial = MakeShared<UE::Geometry::FDynamicMeshAABBTree3, ESPMode::ThreadSafe>(OriginalMesh.Get(), true);
 
 		double InitialMeshArea = 0;
 		for (int tid : OriginalMesh->TriangleIndicesItr())
@@ -109,7 +109,7 @@ void UDataprepRemeshOperation::OnExecution_Implementation(const FDataprepContext
 		}
 
 		double TargetTriArea = InitialMeshArea / (double)TargetTriangleCount;
-		double EdgeLen = TriangleUtil::EquilateralEdgeLengthForArea(TargetTriArea);
+		double EdgeLen = UE::Geometry::TriangleUtil::EquilateralEdgeLengthForArea(TargetTriArea);
 		Op.TargetEdgeLength = (double)FMath::RoundToInt(EdgeLen*100.0) / 100.0;
 
 		Op.OriginalMesh = OriginalMesh;
@@ -123,7 +123,7 @@ void UDataprepRemeshOperation::OnExecution_Implementation(const FDataprepContext
 		FDynamicMeshToMeshDescription DynamicMeshToMeshDescription;
 
 		// full conversion if normal topology changed or faces were inverted
-		TUniquePtr<FDynamicMesh3> ResultMesh = Op.ExtractResult();
+		TUniquePtr<UE::Geometry::FDynamicMesh3> ResultMesh = Op.ExtractResult();
 		DynamicMeshToMeshDescription.Convert(ResultMesh.Get(), *MeshDescription);
 
 		UStaticMesh::FCommitMeshDescriptionParams Params;
@@ -166,11 +166,11 @@ void UDataprepSimplifyMeshOperation::OnExecution_Implementation(const FDataprepC
 		FMeshDescriptionToDynamicMesh MeshDescriptionToDynamicMesh;
 		MeshDescriptionToDynamicMesh.Convert(MeshDescription, *OriginalMesh);
 
-		TSharedPtr<FDynamicMeshAABBTree3, ESPMode::ThreadSafe> OriginalMeshSpatial = MakeShared<FDynamicMeshAABBTree3, ESPMode::ThreadSafe>(OriginalMesh.Get(), true);
+		TSharedPtr<UE::Geometry::FDynamicMeshAABBTree3, ESPMode::ThreadSafe> OriginalMeshSpatial = MakeShared<UE::Geometry::FDynamicMeshAABBTree3, ESPMode::ThreadSafe>(OriginalMesh.Get(), true);
 
 		TSharedPtr<FMeshDescription, ESPMode::ThreadSafe> OriginalMeshDescription = MakeShared<FMeshDescription, ESPMode::ThreadSafe>(*StaticMesh->GetMeshDescription(0));
 
-		FSimplifyMeshOp Op;
+		UE::Geometry::FSimplifyMeshOp Op;
 
 		Op.bDiscardAttributes = bDiscardAttributes;
 		Op.bPreventNormalFlips = true;
@@ -198,7 +198,7 @@ void UDataprepSimplifyMeshOperation::OnExecution_Implementation(const FDataprepC
 		FDynamicMeshToMeshDescription DynamicMeshToMeshDescription;
 
 		// Convert back to mesh description
-		TUniquePtr<FDynamicMesh3> ResultMesh = Op.ExtractResult();
+		TUniquePtr<UE::Geometry::FDynamicMesh3> ResultMesh = Op.ExtractResult();
 		DynamicMeshToMeshDescription.Convert(ResultMesh.Get(), *MeshDescription);
 
 		UStaticMesh::FCommitMeshDescriptionParams Params;
@@ -270,22 +270,22 @@ void UDataprepBakeTransformOperation::OnExecution_Implementation(const FDataprep
 		}
 	}
 
-	TArray<FTransform3d> BakedTransforms;
+	TArray<UE::Geometry::FTransform3d> BakedTransforms;
 	for (int32 ComponentIdx = 0; ComponentIdx < StaticMeshComponents.Num(); ComponentIdx++)
 	{
-		FTransform3d ComponentToWorld(StaticMeshComponents[ComponentIdx]->GetComponentTransform());
-		FTransform3d ToBakePart = FTransform3d::Identity();
-		FTransform3d NewWorldPart = ComponentToWorld;
+		UE::Geometry::FTransform3d ComponentToWorld(StaticMeshComponents[ComponentIdx]->GetComponentTransform());
+		UE::Geometry::FTransform3d ToBakePart = UE::Geometry::FTransform3d::Identity();
+		UE::Geometry::FTransform3d NewWorldPart = ComponentToWorld;
 
 		if (MapToFirstOccurrences[ComponentIdx] < ComponentIdx)
 		{
 			ToBakePart = BakedTransforms[MapToFirstOccurrences[ComponentIdx]];
 			BakedTransforms.Add(ToBakePart);
 			// Try to invert baked transform
-			NewWorldPart = FTransform3d(
+			NewWorldPart = UE::Geometry::FTransform3d(
 				NewWorldPart.GetRotation() * ToBakePart.GetRotation().Inverse(),
 				NewWorldPart.GetTranslation(),
-				NewWorldPart.GetScale() * FTransform3d::GetSafeScaleReciprocal(ToBakePart.GetScale())
+				NewWorldPart.GetScale() * UE::Geometry::FTransform3d::GetSafeScaleReciprocal(ToBakePart.GetScale())
 			);
 			NewWorldPart.SetTranslation(NewWorldPart.GetTranslation() - NewWorldPart.TransformVector(ToBakePart.GetTranslation()));
 		}
@@ -294,14 +294,14 @@ void UDataprepBakeTransformOperation::OnExecution_Implementation(const FDataprep
 			if (bBakeRotation)
 			{
 				ToBakePart.SetRotation(ComponentToWorld.GetRotation());
-				NewWorldPart.SetRotation(FQuaterniond::Identity());
+				NewWorldPart.SetRotation(UE::Geometry::FQuaterniond::Identity());
 			}
-			FVector3d ScaleVec = ComponentToWorld.GetScale();
+			UE::Geometry::FVector3d ScaleVec = ComponentToWorld.GetScale();
 
-			FVector3d AbsScales(FMathd::Abs(ScaleVec.X), FMathd::Abs(ScaleVec.Y), FMathd::Abs(ScaleVec.Z));
+			UE::Geometry::FVector3d AbsScales(FMathd::Abs(ScaleVec.X), FMathd::Abs(ScaleVec.Y), FMathd::Abs(ScaleVec.Z));
 			double RemainingUniformScale = AbsScales.X;
 			{
-				FVector3d Dists;
+				UE::Geometry::FVector3d Dists;
 				for (int SubIdx = 0; SubIdx < 3; SubIdx++)
 				{
 					int OtherA = (SubIdx + 1) % 3;
@@ -326,12 +326,12 @@ void UDataprepBakeTransformOperation::OnExecution_Implementation(const FDataprep
 			{
 			case EBakeScaleMethod::BakeFullScale:
 				ToBakePart.SetScale(ScaleVec);
-				NewWorldPart.SetScale(FVector3d::One());
+				NewWorldPart.SetScale(UE::Geometry::FVector3d::One());
 				break;
 			case EBakeScaleMethod::BakeNonuniformScale:
 				check(RemainingUniformScale > FLT_MIN); // avoid baking a ~zero scale
 				ToBakePart.SetScale(ScaleVec / RemainingUniformScale);
-				NewWorldPart.SetScale(FVector3d(RemainingUniformScale, RemainingUniformScale, RemainingUniformScale));
+				NewWorldPart.SetScale(UE::Geometry::FVector3d(RemainingUniformScale, RemainingUniformScale, RemainingUniformScale));
 				break;
 			case EBakeScaleMethod::DoNotBakeScale:
 				break;
@@ -351,8 +351,8 @@ void UDataprepBakeTransformOperation::OnExecution_Implementation(const FDataprep
 			if (bRecenterPivot)
 			{
 				FBox BBox = MeshDescription->ComputeBoundingBox();
-				FVector3d Center(BBox.GetCenter());
-				FFrame3d LocalFrame(Center);
+				UE::Geometry::FVector3d Center(BBox.GetCenter());
+				UE::Geometry::FFrame3d LocalFrame(Center);
 				ToBakePart.SetTranslation(ToBakePart.GetTranslation() - Center);
 				NewWorldPart.SetTranslation(NewWorldPart.GetTranslation() + NewWorldPart.TransformVector(Center));
 			}
@@ -404,7 +404,7 @@ void UDataprepWeldEdgesOperation::OnExecution_Implementation(const FDataprepCont
 		FMeshDescriptionToDynamicMesh MeshDescriptionToDynamicMesh;
 		MeshDescriptionToDynamicMesh.Convert(MeshDescription, *TargetMesh);
 
-		FMergeCoincidentMeshEdges Merger(TargetMesh.Get());
+		UE::Geometry::FMergeCoincidentMeshEdges Merger(TargetMesh.Get());
 		Merger.MergeSearchTolerance = Tolerance;
 		Merger.OnlyUniquePairs = bOnlyUnique;
 		
@@ -413,7 +413,7 @@ void UDataprepWeldEdgesOperation::OnExecution_Implementation(const FDataprepCont
 			continue;
 		}
 
-		if (TargetMesh->CheckValidity(true, EValidityCheckFailMode::ReturnOnly) == false)
+		if (TargetMesh->CheckValidity(true, UE::Geometry::EValidityCheckFailMode::ReturnOnly) == false)
 		{
 			continue; // Target mesh is invalid
 		}
@@ -555,14 +555,14 @@ TUniquePtr<FDynamicMesh3> UDataprepPlaneCutOperation::CutStaticMesh(const FTrans
 	MeshDescriptionToDynamicMesh.Convert(MeshDescription, *OriginalMesh);
 
 	OriginalMesh->EnableAttributes();
-	TDynamicMeshScalarTriangleAttribute<int>* SubObjectIDs = new TDynamicMeshScalarTriangleAttribute<int>(OriginalMesh.Get());
+	UE::Geometry::TDynamicMeshScalarTriangleAttribute<int>* SubObjectIDs = new UE::Geometry::TDynamicMeshScalarTriangleAttribute<int>(OriginalMesh.Get());
 	SubObjectIDs->Initialize(0);
-	OriginalMesh->Attributes()->AttachAttribute(FPlaneCutOp::ObjectIndexAttribute, SubObjectIDs);
+	OriginalMesh->Attributes()->AttachAttribute(UE::Geometry::FPlaneCutOp::ObjectIndexAttribute, SubObjectIDs);
 
 	// Store a UV scale based on the original mesh bounds (we don't want to recompute this between cuts b/c we want consistent UV scale)
 	const float MeshUVScaleFactor = 1.0 / OriginalMesh->GetBounds().MaxDim();
 
-	TUniquePtr<FDynamicMeshOperator> CutOp = MakeNewOperator(InTransform, OriginalMesh, MeshUVScaleFactor);
+	TUniquePtr<UE::Geometry::FDynamicMeshOperator> CutOp = MakeNewOperator(InTransform, OriginalMesh, MeshUVScaleFactor);
 
 	FProgressCancel Progress;
 	CutOp->CalculateResult(&Progress);
@@ -612,11 +612,11 @@ void UDataprepPlaneCutOperation::PerformCutting(
 		{
 			// Export separated pieces as new mesh assets
 
-			TDynamicMeshScalarTriangleAttribute<int>* SubMeshIDs =
-				static_cast<TDynamicMeshScalarTriangleAttribute<int>*>(UseMesh->Attributes()->GetAttachedAttribute(
-					FPlaneCutOp::ObjectIndexAttribute));
-			TArray<FDynamicMesh3>& SplitMeshes = AllSplitMeshes[OrigMeshIdx];
-			bool bWasSplit = FDynamicMeshEditor::SplitMesh(UseMesh, SplitMeshes, [SubMeshIDs](int TID)
+			UE::Geometry::TDynamicMeshScalarTriangleAttribute<int>* SubMeshIDs =
+				static_cast<UE::Geometry::TDynamicMeshScalarTriangleAttribute<int>*>(UseMesh->Attributes()->GetAttachedAttribute(
+					UE::Geometry::FPlaneCutOp::ObjectIndexAttribute));
+			TArray<UE::Geometry::FDynamicMesh3>& SplitMeshes = AllSplitMeshes[OrigMeshIdx];
+			bool bWasSplit = UE::Geometry::FDynamicMeshEditor::SplitMesh(UseMesh, SplitMeshes, [SubMeshIDs](int TID)
 			{
 				return SubMeshIDs->GetValue(TID);
 			});
@@ -752,12 +752,12 @@ void UDataprepPlaneCutOperation::PerformCutting(
 	}
 }
 
-TUniquePtr<FDynamicMeshOperator> UDataprepPlaneCutOperation::MakeNewOperator(
+TUniquePtr<UE::Geometry::FDynamicMeshOperator> UDataprepPlaneCutOperation::MakeNewOperator(
 	const FTransform& InMeshLocalToWorld, 
-	TSharedPtr<FDynamicMesh3, ESPMode::ThreadSafe> InOriginalMesh, 
+	TSharedPtr<UE::Geometry::FDynamicMesh3, ESPMode::ThreadSafe> InOriginalMesh,
 	float InMeshUVScaleFactor)
 {
-	TUniquePtr<FPlaneCutOp> CutOp = MakeUnique<FPlaneCutOp>();
+	TUniquePtr<UE::Geometry::FPlaneCutOp> CutOp = MakeUnique<UE::Geometry::FPlaneCutOp>();
 	CutOp->bFillCutHole = bFillCutHole;
 	CutOp->bFillSpans = false;
 
@@ -787,7 +787,7 @@ TUniquePtr<FDynamicMeshOperator> UDataprepPlaneCutOperation::MakeNewOperator(
 	LocalToWorld.SetScale3D(LocalToWorldScale);
 	FTransform WorldToLocal = LocalToWorld.Inverse();
 	FVector LocalOrigin = WorldToLocal.TransformPosition(CutPlaneOrigin);
-	FTransform3d W2LForNormal(WorldToLocal);
+	UE::Geometry::FTransform3d W2LForNormal(WorldToLocal);
 	FVector LocalNormal = (FVector)W2LForNormal.TransformNormal(WorldNormal);
 	FVector BackTransformed = LocalToWorld.TransformVector(LocalNormal);
 	float NormalScaleFactor = FVector::DotProduct(BackTransformed, WorldNormal);
diff --git a/Engine/Plugins/Experimental/Enterprise/DataprepGeometryOperations/Source/Private/DataprepGeometryOperations.h b/Engine/Plugins/Experimental/Enterprise/DataprepGeometryOperations/Source/Private/DataprepGeometryOperations.h
--- a/Engine/Plugins/Experimental/Enterprise/DataprepGeometryOperations/Source/Private/DataprepGeometryOperations.h
+++ b/Engine/Plugins/Experimental/Enterprise/DataprepGeometryOperations/Source/Private/DataprepGeometryOperations.h
@@ -256,8 +256,8 @@ private:
 		TArray<UObject*>& OutModifiedStaticMeshes);
 
 	// Borrowed from UPlaneCutOperatorFactory::MakeNewOperator
-	TUniquePtr<FDynamicMeshOperator> MakeNewOperator(
+	TUniquePtr<UE::Geometry::FDynamicMeshOperator> MakeNewOperator(
 		const FTransform& InMeshLocalToWorld, 
-		TSharedPtr<FDynamicMesh3, ESPMode::ThreadSafe> InOriginalMesh,
+		TSharedPtr<UE::Geometry::FDynamicMesh3, ESPMode::ThreadSafe> InOriginalMesh,
 		float InMeshUVScaleFactor);
 };
diff --git a/Engine/Plugins/Experimental/GeometryFlow/Source/GeometryFlowCore/Public/GeometryFlowCoreNodes.h b/Engine/Plugins/Experimental/GeometryFlow/Source/GeometryFlowCore/Public/GeometryFlowCoreNodes.h
--- a/Engine/Plugins/Experimental/GeometryFlow/Source/GeometryFlowCore/Public/GeometryFlowCoreNodes.h
+++ b/Engine/Plugins/Experimental/GeometryFlow/Source/GeometryFlowCore/Public/GeometryFlowCoreNodes.h
@@ -11,6 +11,9 @@ namespace UE
 namespace GeometryFlow
 {
 
+using namespace UE::Geometry;
+using UE::Geometry::FVector3d;
+using UE::Geometry::FVector3f;
 
 
 template<typename T, int32 StorageTypeIdentifier>
diff --git a/Engine/Plugins/Experimental/GeometryFlow/Source/GeometryFlowMeshProcessing/Private/DataTypes/MeshImageBakingData.cpp b/Engine/Plugins/Experimental/GeometryFlow/Source/GeometryFlowMeshProcessing/Private/DataTypes/MeshImageBakingData.cpp
--- a/Engine/Plugins/Experimental/GeometryFlow/Source/GeometryFlowMeshProcessing/Private/DataTypes/MeshImageBakingData.cpp
+++ b/Engine/Plugins/Experimental/GeometryFlow/Source/GeometryFlowMeshProcessing/Private/DataTypes/MeshImageBakingData.cpp
@@ -2,6 +2,7 @@
 
 #include "DataTypes/MeshImageBakingData.h"
 
+using namespace UE::Geometry;
 using namespace UE::GeometryFlow;
 
 void FMakeMeshBakingCacheNode::Evaluate(
diff --git a/Engine/Plugins/Experimental/GeometryFlow/Source/GeometryFlowMeshProcessing/Private/MeshBakingNodes/BakeMeshNormalMapNode.cpp b/Engine/Plugins/Experimental/GeometryFlow/Source/GeometryFlowMeshProcessing/Private/MeshBakingNodes/BakeMeshNormalMapNode.cpp
--- a/Engine/Plugins/Experimental/GeometryFlow/Source/GeometryFlowMeshProcessing/Private/MeshBakingNodes/BakeMeshNormalMapNode.cpp
+++ b/Engine/Plugins/Experimental/GeometryFlow/Source/GeometryFlowMeshProcessing/Private/MeshBakingNodes/BakeMeshNormalMapNode.cpp
@@ -5,6 +5,7 @@
 
 #include "Sampling/MeshNormalMapBaker.h"
 
+using namespace UE::Geometry;
 using namespace UE::GeometryFlow;
 
 void FBakeMeshNormalMapNode::Evaluate(
diff --git a/Engine/Plugins/Experimental/GeometryFlow/Source/GeometryFlowMeshProcessing/Private/MeshProcessingNodes/MeshRecalculateUVsNode.cpp b/Engine/Plugins/Experimental/GeometryFlow/Source/GeometryFlowMeshProcessing/Private/MeshProcessingNodes/MeshRecalculateUVsNode.cpp
--- a/Engine/Plugins/Experimental/GeometryFlow/Source/GeometryFlowMeshProcessing/Private/MeshProcessingNodes/MeshRecalculateUVsNode.cpp
+++ b/Engine/Plugins/Experimental/GeometryFlow/Source/GeometryFlowMeshProcessing/Private/MeshProcessingNodes/MeshRecalculateUVsNode.cpp
@@ -7,7 +7,7 @@
 #include "Parameterization/MeshUVTransforms.h"
 
 
-
+using namespace UE::Geometry;
 using namespace UE::GeometryFlow;
 
 
diff --git a/Engine/Plugins/Experimental/GeometryFlow/Source/GeometryFlowMeshProcessing/Public/DataTypes/CollisionGeometryData.h b/Engine/Plugins/Experimental/GeometryFlow/Source/GeometryFlowMeshProcessing/Public/DataTypes/CollisionGeometryData.h
--- a/Engine/Plugins/Experimental/GeometryFlow/Source/GeometryFlowMeshProcessing/Public/DataTypes/CollisionGeometryData.h
+++ b/Engine/Plugins/Experimental/GeometryFlow/Source/GeometryFlowMeshProcessing/Public/DataTypes/CollisionGeometryData.h
@@ -20,7 +20,7 @@ struct FCollisionGeometry
 {
 	DECLARE_GEOMETRYFLOW_DATA_TYPE_IDENTIFIER(EMeshProcessingDataTypes::CollisionGeometry);
 
-	FSimpleShapeSet3d Geometry;
+	UE::Geometry::FSimpleShapeSet3d Geometry;
 };
 
 
diff --git a/Engine/Plugins/Experimental/GeometryFlow/Source/GeometryFlowMeshProcessing/Public/DataTypes/DynamicMeshData.h b/Engine/Plugins/Experimental/GeometryFlow/Source/GeometryFlowMeshProcessing/Public/DataTypes/DynamicMeshData.h
--- a/Engine/Plugins/Experimental/GeometryFlow/Source/GeometryFlowMeshProcessing/Public/DataTypes/DynamicMeshData.h
+++ b/Engine/Plugins/Experimental/GeometryFlow/Source/GeometryFlowMeshProcessing/Public/DataTypes/DynamicMeshData.h
@@ -12,6 +12,8 @@ namespace UE
 namespace GeometryFlow
 {
 
+using namespace UE::Geometry;
+
 // declares FDataDynamicMesh, FDynamicMeshInput, FDynamicMeshOutput, FDynamicMeshSourceNode
 GEOMETRYFLOW_DECLARE_BASIC_TYPES(DynamicMesh, FDynamicMesh3, (int)EMeshProcessingDataTypes::DynamicMesh)
 
diff --git a/Engine/Plugins/Experimental/GeometryFlow/Source/GeometryFlowMeshProcessing/Public/DataTypes/MeshImageBakingData.h b/Engine/Plugins/Experimental/GeometryFlow/Source/GeometryFlowMeshProcessing/Public/DataTypes/MeshImageBakingData.h
--- a/Engine/Plugins/Experimental/GeometryFlow/Source/GeometryFlowMeshProcessing/Public/DataTypes/MeshImageBakingData.h
+++ b/Engine/Plugins/Experimental/GeometryFlow/Source/GeometryFlowMeshProcessing/Public/DataTypes/MeshImageBakingData.h
@@ -15,7 +15,7 @@ namespace UE
 namespace GeometryFlow
 {
 
-
+using namespace UE::Geometry;
 
 struct GEOMETRYFLOWMESHPROCESSING_API FMeshBakingCache
 {
@@ -26,7 +26,7 @@ struct GEOMETRYFLOWMESHPROCESSING_API FMeshBakingCache
 	FDynamicMeshAABBTree3 DetailSpatial;
 	FDynamicMesh3 TargetMesh;
 
-	FMeshImageBakingCache BakeCache;
+	UE::Geometry::FMeshImageBakingCache BakeCache;
 };
 
 
diff --git a/Engine/Plugins/Experimental/GeometryFlow/Source/GeometryFlowMeshProcessing/Public/MeshProcessingNodes/MeshProcessingBaseNodes.h b/Engine/Plugins/Experimental/GeometryFlow/Source/GeometryFlowMeshProcessing/Public/MeshProcessingNodes/MeshProcessingBaseNodes.h
--- a/Engine/Plugins/Experimental/GeometryFlow/Source/GeometryFlowMeshProcessing/Public/MeshProcessingNodes/MeshProcessingBaseNodes.h
+++ b/Engine/Plugins/Experimental/GeometryFlow/Source/GeometryFlowMeshProcessing/Public/MeshProcessingNodes/MeshProcessingBaseNodes.h
@@ -12,10 +12,7 @@ namespace UE
 namespace GeometryFlow
 {
 
-
-
-
-
+using namespace UE::Geometry;
 
 
 class FProcessMeshBaseNode : public FNode
diff --git a/Engine/Plugins/Experimental/GeometryProcessing/Source/DynamicMesh/Private/DynamicMesh3.cpp b/Engine/Plugins/Experimental/GeometryProcessing/Source/DynamicMesh/Private/DynamicMesh3.cpp
--- a/Engine/Plugins/Experimental/GeometryProcessing/Source/DynamicMesh/Private/DynamicMesh3.cpp
+++ b/Engine/Plugins/Experimental/GeometryProcessing/Source/DynamicMesh/Private/DynamicMesh3.cpp
@@ -5,6 +5,9 @@
 #include "Generators/MeshShapeGenerator.h"
 #include "Templates/UniquePtr.h"
 
+#include "ExplicitUseGeometryMathTypes.h"		// using UE::Geometry::(math types)
+using namespace UE::Geometry;
+
 // NB: These have to be here until C++17 allows inline variables
 constexpr int       FDynamicMesh3::InvalidID;
 constexpr int       FDynamicMesh3::NonManifoldID;
diff --git a/Engine/Plugins/Experimental/GeometryProcessing/Source/DynamicMesh/Private/DynamicMesh3_Edits.cpp b/Engine/Plugins/Experimental/GeometryProcessing/Source/DynamicMesh/Private/DynamicMesh3_Edits.cpp
--- a/Engine/Plugins/Experimental/GeometryProcessing/Source/DynamicMesh/Private/DynamicMesh3_Edits.cpp
+++ b/Engine/Plugins/Experimental/GeometryProcessing/Source/DynamicMesh/Private/DynamicMesh3_Edits.cpp
@@ -3,7 +3,8 @@
 #include "DynamicMesh3.h"
 #include "DynamicMeshAttributeSet.h"
 
-
+#include "ExplicitUseGeometryMathTypes.h"		// using UE::Geometry::(math types)
+using namespace UE::Geometry;
 
 
 int FDynamicMesh3::AppendVertex(const FVertexInfo& VtxInfo)
diff --git a/Engine/Plugins/Experimental/GeometryProcessing/Source/DynamicMesh/Private/DynamicMesh3_Queries.cpp b/Engine/Plugins/Experimental/GeometryProcessing/Source/DynamicMesh/Private/DynamicMesh3_Queries.cpp
--- a/Engine/Plugins/Experimental/GeometryProcessing/Source/DynamicMesh/Private/DynamicMesh3_Queries.cpp
+++ b/Engine/Plugins/Experimental/GeometryProcessing/Source/DynamicMesh/Private/DynamicMesh3_Queries.cpp
@@ -3,6 +3,8 @@
 #include "DynamicMesh3.h"
 #include "Async/ParallelFor.h"
 
+#include "ExplicitUseGeometryMathTypes.h"		// using UE::Geometry::(math types)
+using namespace UE::Geometry;
 
 FIndex2i FDynamicMesh3::GetEdgeOpposingV(int eID) const
 {
diff --git a/Engine/Plugins/Experimental/GeometryProcessing/Source/DynamicMesh/Private/DynamicMeshAttributeSet.cpp b/Engine/Plugins/Experimental/GeometryProcessing/Source/DynamicMesh/Private/DynamicMeshAttributeSet.cpp
--- a/Engine/Plugins/Experimental/GeometryProcessing/Source/DynamicMesh/Private/DynamicMeshAttributeSet.cpp
+++ b/Engine/Plugins/Experimental/GeometryProcessing/Source/DynamicMesh/Private/DynamicMeshAttributeSet.cpp
@@ -4,7 +4,8 @@
 #include "IndexTypes.h"
 #include "Async/ParallelFor.h"
 
-
+#include "ExplicitUseGeometryMathTypes.h"		// using UE::Geometry::(math types)
+using namespace UE::Geometry;
 
 
 void FDynamicMeshAttributeSet::Copy(const FDynamicMeshAttributeSet& Copy)
diff --git a/Engine/Plugins/Experimental/GeometryProcessing/Source/DynamicMesh/Private/DynamicMeshChangeTracker.cpp b/Engine/Plugins/Experimental/GeometryProcessing/Source/DynamicMesh/Private/DynamicMeshChangeTracker.cpp
--- a/Engine/Plugins/Experimental/GeometryProcessing/Source/DynamicMesh/Private/DynamicMeshChangeTracker.cpp
+++ b/Engine/Plugins/Experimental/GeometryProcessing/Source/DynamicMesh/Private/DynamicMeshChangeTracker.cpp
@@ -2,7 +2,8 @@
 
 #include "DynamicMeshChangeTracker.h"
 
-
+#include "ExplicitUseGeometryMathTypes.h"		// using UE::Geometry::(math types)
+using namespace UE::Geometry;
 
 //
 //
@@ -138,7 +139,10 @@ void TDynamicMeshAttributeChange<RealType,ElementSize>::ApplyReplaceChange(TDyna
 
 
 
-
+namespace UE
+{
+namespace Geometry
+{
 
 // because ::ApplyReplaceChange() is in the .cpp we need to declare possible attribute types here
 template class DYNAMICMESH_API TDynamicMeshAttributeChange<float, 1>;
@@ -151,7 +155,8 @@ template class DYNAMICMESH_API TDynamicMeshAttributeChange<float, 3>;
 template class DYNAMICMESH_API TDynamicMeshAttributeChange<double, 3>;
 template class DYNAMICMESH_API TDynamicMeshAttributeChange<int, 3>;
 
-
+} // end namespace UE::Geometry
+} // end namespace UE
 
 
 
diff --git a/Engine/Plugins/Experimental/GeometryProcessing/Source/DynamicMesh/Private/DynamicMeshOverlay.cpp b/Engine/Plugins/Experimental/GeometryProcessing/Source/DynamicMesh/Private/DynamicMeshOverlay.cpp
--- a/Engine/Plugins/Experimental/GeometryProcessing/Source/DynamicMesh/Private/DynamicMeshOverlay.cpp
+++ b/Engine/Plugins/Experimental/GeometryProcessing/Source/DynamicMesh/Private/DynamicMeshOverlay.cpp
@@ -3,7 +3,8 @@
 #include "DynamicMeshOverlay.h"
 #include "DynamicMesh3.h"
 
-
+#include "ExplicitUseGeometryMathTypes.h"		// using UE::Geometry::(math types)
+using namespace UE::Geometry;
 
 
 template<typename RealType, int ElementSize>
@@ -1388,7 +1389,10 @@ bool TDynamicMeshOverlay<RealType, ElementSize>::CheckValidity(bool bAllowNonMan
 
 
 
-
+namespace UE
+{
+namespace Geometry
+{
 
 // These are explicit instantiations of the templates that are exported from the shared lib.
 // Only these instantiations of the template can be used.
@@ -1408,4 +1412,7 @@ template class DYNAMICMESH_API TDynamicMeshVectorOverlay<double, 2, FVector2d>;
 template class DYNAMICMESH_API TDynamicMeshVectorOverlay<int, 2, FVector2i>;
 template class DYNAMICMESH_API TDynamicMeshVectorOverlay<float, 3, FVector3f>;
 template class DYNAMICMESH_API TDynamicMeshVectorOverlay<double, 3, FVector3d>;
-template class DYNAMICMESH_API TDynamicMeshVectorOverlay<int, 3, FVector3i>;
\ No newline at end of file
+template class DYNAMICMESH_API TDynamicMeshVectorOverlay<int, 3, FVector3i>;
+
+} // end namespace UE::Geometry
+} // end namespace UE
\ No newline at end of file
diff --git a/Engine/Plugins/Experimental/GeometryProcessing/Source/DynamicMesh/Private/EdgeLoop.cpp b/Engine/Plugins/Experimental/GeometryProcessing/Source/DynamicMesh/Private/EdgeLoop.cpp
--- a/Engine/Plugins/Experimental/GeometryProcessing/Source/DynamicMesh/Private/EdgeLoop.cpp
+++ b/Engine/Plugins/Experimental/GeometryProcessing/Source/DynamicMesh/Private/EdgeLoop.cpp
@@ -3,6 +3,8 @@
 
 #include "EdgeLoop.h"
 
+#include "ExplicitUseGeometryMathTypes.h"		// using UE::Geometry::(math types)
+using namespace UE::Geometry;
 
 void FEdgeLoop::Initialize(const FDynamicMesh3* mesh, const TArray<int>& vertices, const TArray<int> & edges, const TArray<int>* BowtieVerticesIn)
 {
diff --git a/Engine/Plugins/Experimental/GeometryProcessing/Source/DynamicMesh/Private/EdgeSpan.cpp b/Engine/Plugins/Experimental/GeometryProcessing/Source/DynamicMesh/Private/EdgeSpan.cpp
--- a/Engine/Plugins/Experimental/GeometryProcessing/Source/DynamicMesh/Private/EdgeSpan.cpp
+++ b/Engine/Plugins/Experimental/GeometryProcessing/Source/DynamicMesh/Private/EdgeSpan.cpp
@@ -2,6 +2,9 @@
 
 #include "EdgeSpan.h"
 
+#include "ExplicitUseGeometryMathTypes.h"		// using UE::Geometry::(math types)
+using namespace UE::Geometry;
+
 void FEdgeSpan::Initialize(const FDynamicMesh3* mesh, const TArray<int>& vertices, const TArray<int> & edges, const TArray<int>* BowtieVerticesIn)
 {
 	Mesh = mesh;
diff --git a/Engine/Plugins/Experimental/GeometryProcessing/Source/DynamicMesh/Private/MeshSimplification.cpp b/Engine/Plugins/Experimental/GeometryProcessing/Source/DynamicMesh/Private/MeshSimplification.cpp
--- a/Engine/Plugins/Experimental/GeometryProcessing/Source/DynamicMesh/Private/MeshSimplification.cpp
+++ b/Engine/Plugins/Experimental/GeometryProcessing/Source/DynamicMesh/Private/MeshSimplification.cpp
@@ -7,7 +7,8 @@
 #include "Async/ParallelFor.h"
 #include "Templates/UnrealTypeTraits.h"
 
-
+#include "ExplicitUseGeometryMathTypes.h"		// using UE::Geometry::(math types)
+using namespace UE::Geometry;
 
 template <typename QuadricErrorType>
 QuadricErrorType TMeshSimplification<QuadricErrorType>::ComputeFaceQuadric(const int tid, FVector3d& nface, FVector3d& c, double& Area) const
@@ -1595,7 +1596,10 @@ void TMeshSimplification<FAttrBasedQuadricErrord>::OnEdgeCollapse(int edgeID, in
 
 }
 
-
+namespace UE
+{
+namespace Geometry
+{
 
 // These are explicit instantiations of the templates that are exported from the shared lib.
 // Only these instantiations of the template can be used.
@@ -1603,3 +1607,6 @@ void TMeshSimplification<FAttrBasedQuadricErrord>::OnEdgeCollapse(int edgeID, in
 template class DYNAMICMESH_API TMeshSimplification< FAttrBasedQuadricErrord >;
 template class DYNAMICMESH_API TMeshSimplification< FVolPresQuadricErrord >;
 template class DYNAMICMESH_API TMeshSimplification< FQuadricErrord >;
+
+} // end namespace UE::Geometry
+} // end namespace UE
\ No newline at end of file
diff --git a/Engine/Plugins/Experimental/GeometryProcessing/Source/DynamicMesh/Private/MeshTangents.cpp b/Engine/Plugins/Experimental/GeometryProcessing/Source/DynamicMesh/Private/MeshTangents.cpp
--- a/Engine/Plugins/Experimental/GeometryProcessing/Source/DynamicMesh/Private/MeshTangents.cpp
+++ b/Engine/Plugins/Experimental/GeometryProcessing/Source/DynamicMesh/Private/MeshTangents.cpp
@@ -3,6 +3,8 @@
 #include "MeshTangents.h"
 #include "Async/ParallelFor.h"
 
+#include "ExplicitUseGeometryMathTypes.h"		// using UE::Geometry::(math types)
+using namespace UE::Geometry;
 
 template<typename RealType>
 void TMeshTangents<RealType>::SetTangentCount(int Count, bool bClearToZero)
@@ -477,10 +479,15 @@ void TMeshTangents<RealType>::ComputeMikkTStyleTangents(const FDynamicMeshNormal
 	}
 }
 
+namespace UE
+{
+namespace Geometry
+{
 
 template class DYNAMICMESH_API TMeshTangents<float>;
 template class DYNAMICMESH_API TMeshTangents<double>;
 
-
+} // end namespace UE::Geometry
+} // end namespace UE
 
 
diff --git a/Engine/Plugins/Experimental/GeometryProcessing/Source/DynamicMesh/Private/Operations/EmbedSurfacePath.cpp b/Engine/Plugins/Experimental/GeometryProcessing/Source/DynamicMesh/Private/Operations/EmbedSurfacePath.cpp
--- a/Engine/Plugins/Experimental/GeometryProcessing/Source/DynamicMesh/Private/Operations/EmbedSurfacePath.cpp
+++ b/Engine/Plugins/Experimental/GeometryProcessing/Source/DynamicMesh/Private/Operations/EmbedSurfacePath.cpp
@@ -6,8 +6,11 @@
 #include "TriangleTypes.h"
 #include "SegmentTypes.h"
 #include "Distance/DistPoint3Triangle3.h"
-
 #include "DynamicMesh3.h"
+#include "ExplicitUseGeometryMathTypes.h"		// using UE::Geometry::(math types)
+
+#include "ExplicitUseGeometryMathTypes.h"		// using UE::Geometry::(math types)
+using namespace UE::Geometry;
 
 FVector3d FMeshSurfacePoint::Pos(const FDynamicMesh3 *Mesh) const
 {
@@ -781,7 +784,7 @@ bool AddEdgesOnPath(const FDynamicMesh3& Mesh, FFrame3d Frame, const TArray<int>
 }
 
 
-bool EmbedProjectedPaths(FDynamicMesh3* Mesh, const TArrayView<const int> StartTriIDs, FFrame3d Frame, const TArrayView<const TArray<FVector2d>> AllPaths, TArray<TArray<int>>& OutAllPathVertices, TArray<TArray<int>>& OutAllVertexCorrespondence, bool bClosePaths, FMeshFaceSelection* EnclosedFaces, double PtSnapVertexOrEdgeThresholdSq)
+bool UE::Geometry::EmbedProjectedPaths(FDynamicMesh3* Mesh, const TArrayView<const int> StartTriIDs, FFrame3d Frame, const TArrayView<const TArray<FVector2d>> AllPaths, TArray<TArray<int>>& OutAllPathVertices, TArray<TArray<int>>& OutAllVertexCorrespondence, bool bClosePaths, FMeshFaceSelection* EnclosedFaces, double PtSnapVertexOrEdgeThresholdSq)
 {
 	check(Mesh);
 	check(AllPaths.Num() == StartTriIDs.Num());
@@ -945,7 +948,7 @@ bool EmbedProjectedPaths(FDynamicMesh3* Mesh, const TArrayView<const int> StartT
 }
 
 
-bool EmbedProjectedPath(FDynamicMesh3* Mesh, int StartTriID, FFrame3d Frame, const TArray<FVector2d>& Path2D, TArray<int>& OutPathVertices, TArray<int>& OutVertexCorrespondence, bool bClosePath, FMeshFaceSelection* EnclosedFaces, double PtSnapVertexOrEdgeThresholdSq)
+bool UE::Geometry::EmbedProjectedPath(FDynamicMesh3* Mesh, int StartTriID, FFrame3d Frame, const TArray<FVector2d>& Path2D, TArray<int>& OutPathVertices, TArray<int>& OutVertexCorrespondence, bool bClosePath, FMeshFaceSelection* EnclosedFaces, double PtSnapVertexOrEdgeThresholdSq)
 {
 	TArrayView<int> StartTriIDs(&StartTriID, 1);
 	const TArrayView<const TArray<FVector2d>> AllPaths(&Path2D, 1);
diff --git a/Engine/Plugins/Experimental/GeometryProcessing/Source/DynamicMesh/Private/Operations/MeshMeshCut.cpp b/Engine/Plugins/Experimental/GeometryProcessing/Source/DynamicMesh/Private/Operations/MeshMeshCut.cpp
--- a/Engine/Plugins/Experimental/GeometryProcessing/Source/DynamicMesh/Private/Operations/MeshMeshCut.cpp
+++ b/Engine/Plugins/Experimental/GeometryProcessing/Source/DynamicMesh/Private/Operations/MeshMeshCut.cpp
@@ -3,9 +3,11 @@
 // Port of geometry3Sharp MeshMeshCut
 
 #include "Operations/MeshMeshCut.h"
-
 #include "Operations/EmbedSurfacePath.h"
 
+#include "ExplicitUseGeometryMathTypes.h"		// using UE::Geometry::(math types)
+using namespace UE::Geometry;
+
 namespace MeshCut
 {
 	enum class EVertexType
diff --git a/Engine/Plugins/Experimental/GeometryProcessing/Source/DynamicMesh/Private/Polygroups/PolygroupSet.cpp b/Engine/Plugins/Experimental/GeometryProcessing/Source/DynamicMesh/Private/Polygroups/PolygroupSet.cpp
--- a/Engine/Plugins/Experimental/GeometryProcessing/Source/DynamicMesh/Private/Polygroups/PolygroupSet.cpp
+++ b/Engine/Plugins/Experimental/GeometryProcessing/Source/DynamicMesh/Private/Polygroups/PolygroupSet.cpp
@@ -3,6 +3,7 @@
 #include "Polygroups/PolygroupSet.h"
 #include "Polygroups/PolygroupUtil.h"
 
+#include "ExplicitUseGeometryMathTypes.h"		// using UE::Geometry::(math types)
 using namespace UE::Geometry;
 
 
diff --git a/Engine/Plugins/Experimental/GeometryProcessing/Source/DynamicMesh/Private/ShapeApproximation/MeshSimpleShapeApproximation.cpp b/Engine/Plugins/Experimental/GeometryProcessing/Source/DynamicMesh/Private/ShapeApproximation/MeshSimpleShapeApproximation.cpp
--- a/Engine/Plugins/Experimental/GeometryProcessing/Source/DynamicMesh/Private/ShapeApproximation/MeshSimpleShapeApproximation.cpp
+++ b/Engine/Plugins/Experimental/GeometryProcessing/Source/DynamicMesh/Private/ShapeApproximation/MeshSimpleShapeApproximation.cpp
@@ -14,6 +14,9 @@
 #include "Operations/MeshProjectionHull.h"
 #include "Util/ProgressCancel.h"
 
+#include "ExplicitUseGeometryMathTypes.h"		// using UE::Geometry::(math types)
+using namespace UE::Geometry;
+
 
 void FMeshSimpleShapeApproximation::DetectAndCacheSimpleShapeType(const FDynamicMesh3* SourceMesh, FSourceMeshCache& CacheOut)
 {
@@ -95,7 +98,7 @@ struct FSimpleShapeFitsResult
 	FOrientedBox3d Box;
 
 	bool bHaveCapsule = false;
-	FCapsule3d Capsule;
+	UE::Geometry::FCapsule3d Capsule;
 
 	bool bHaveConvex = false;
 	FDynamicMesh3 Convex;
diff --git a/Engine/Plugins/Experimental/GeometryProcessing/Source/DynamicMesh/Private/Solvers/Internal/ConstrainedMeshDeformationSolver.h b/Engine/Plugins/Experimental/GeometryProcessing/Source/DynamicMesh/Private/Solvers/Internal/ConstrainedMeshDeformationSolver.h
--- a/Engine/Plugins/Experimental/GeometryProcessing/Source/DynamicMesh/Private/Solvers/Internal/ConstrainedMeshDeformationSolver.h
+++ b/Engine/Plugins/Experimental/GeometryProcessing/Source/DynamicMesh/Private/Solvers/Internal/ConstrainedMeshDeformationSolver.h
@@ -12,6 +12,11 @@
 #include "MatrixBase.h"
 
 
+namespace UE
+{
+namespace Geometry
+{
+
 class FConstrainedSolver;
 
 /**
@@ -202,4 +207,8 @@ protected:
 	double LaplacianScale = 1.0;
 	bool HasLaplacianScale() const;
 	double GetLaplacianScale(int32 LinearVtxIndex) const;
-};
\ No newline at end of file
+};
+
+
+} // end namespace UE::Geometry
+} // end namespace UE
diff --git a/Engine/Plugins/Experimental/GeometryProcessing/Source/DynamicMesh/Private/Solvers/Internal/ConstrainedMeshDeformers.h b/Engine/Plugins/Experimental/GeometryProcessing/Source/DynamicMesh/Private/Solvers/Internal/ConstrainedMeshDeformers.h
--- a/Engine/Plugins/Experimental/GeometryProcessing/Source/DynamicMesh/Private/Solvers/Internal/ConstrainedMeshDeformers.h
+++ b/Engine/Plugins/Experimental/GeometryProcessing/Source/DynamicMesh/Private/Solvers/Internal/ConstrainedMeshDeformers.h
@@ -4,9 +4,13 @@
 
 #include "CoreMinimal.h"
 #include "DynamicMesh3.h"
-
 #include "ConstrainedMeshDeformationSolver.h"
 
+namespace UE
+{
+namespace Geometry
+{
+
 
 /**
  * FConstrainedMeshDeformer solves detail-preserving mesh deformation problems with arbitrary position constraints.
@@ -51,3 +55,7 @@ private:
 	FSOAPositions LaplacianVectors;
 	FSOAPositions OriginalPositions;
 };
+
+
+} // end namespace UE::Geometry
+} // end namespace UE
\ No newline at end of file
diff --git a/Engine/Plugins/Experimental/GeometryProcessing/Source/DynamicMesh/Private/Solvers/Internal/ConstrainedMeshSmoothers.cpp b/Engine/Plugins/Experimental/GeometryProcessing/Source/DynamicMesh/Private/Solvers/Internal/ConstrainedMeshSmoothers.cpp
--- a/Engine/Plugins/Experimental/GeometryProcessing/Source/DynamicMesh/Private/Solvers/Internal/ConstrainedMeshSmoothers.cpp
+++ b/Engine/Plugins/Experimental/GeometryProcessing/Source/DynamicMesh/Private/Solvers/Internal/ConstrainedMeshSmoothers.cpp
@@ -4,8 +4,8 @@
 #include "MatrixSolver.h"
 #include "ConstrainedPoissonSolver.h"
 
-
-
+#include "ExplicitUseGeometryMathTypes.h"		// using UE::Geometry::(math types)
+using namespace UE::Geometry;
 
 bool FBiHarmonicMeshSmoother::ComputeSmoothedMeshPositions(TArray<FVector3d>& UpdatedPositions)
 {
diff --git a/Engine/Plugins/Experimental/GeometryProcessing/Source/DynamicMesh/Private/Solvers/Internal/ConstrainedMeshSmoothers.h b/Engine/Plugins/Experimental/GeometryProcessing/Source/DynamicMesh/Private/Solvers/Internal/ConstrainedMeshSmoothers.h
--- a/Engine/Plugins/Experimental/GeometryProcessing/Source/DynamicMesh/Private/Solvers/Internal/ConstrainedMeshSmoothers.h
+++ b/Engine/Plugins/Experimental/GeometryProcessing/Source/DynamicMesh/Private/Solvers/Internal/ConstrainedMeshSmoothers.h
@@ -8,6 +8,10 @@
 #include "ConstrainedMeshDeformationSolver.h"
 
 
+namespace UE
+{
+namespace Geometry
+{
 
 
 /**
@@ -63,3 +67,7 @@ public:
 	bool ComputeSmoothedMeshPositions(TArray<FVector3d>& UpdatedPositions);
 
 };
+
+
+} // end namespace UE::Geometry
+} // end namespace UE
\ No newline at end of file
diff --git a/Engine/Plugins/Experimental/GeometryProcessing/Source/DynamicMesh/Private/Solvers/Internal/ConstrainedPoissonSolver.h b/Engine/Plugins/Experimental/GeometryProcessing/Source/DynamicMesh/Private/Solvers/Internal/ConstrainedPoissonSolver.h
--- a/Engine/Plugins/Experimental/GeometryProcessing/Source/DynamicMesh/Private/Solvers/Internal/ConstrainedPoissonSolver.h
+++ b/Engine/Plugins/Experimental/GeometryProcessing/Source/DynamicMesh/Private/Solvers/Internal/ConstrainedPoissonSolver.h
@@ -8,10 +8,13 @@
 #include "FSparseMatrixD.h"
 #include "MatrixSolver.h"
 
-
 // only needed if you enable logging by changing the typedef for FConstrainedSolverTimeLogger
 //#include "ProfilingDebugging/ScopedTimers.h"
 
+namespace UE
+{
+namespace Geometry
+{
 
 
 /**
@@ -377,3 +380,6 @@ private:
 
 };
 
+
+} // end namespace UE::Geometry
+} // end namespace UE
\ No newline at end of file
diff --git a/Engine/Plugins/Experimental/GeometryProcessing/Source/DynamicMesh/Private/Solvers/Internal/FSparseMatrixD.h b/Engine/Plugins/Experimental/GeometryProcessing/Source/DynamicMesh/Private/Solvers/Internal/FSparseMatrixD.h
--- a/Engine/Plugins/Experimental/GeometryProcessing/Source/DynamicMesh/Private/Solvers/Internal/FSparseMatrixD.h
+++ b/Engine/Plugins/Experimental/GeometryProcessing/Source/DynamicMesh/Private/Solvers/Internal/FSparseMatrixD.h
@@ -32,6 +32,7 @@ PRAGMA_DEFAULT_VISIBILITY_END
 #endif
 
 
+
 // NB: The LU solver likes ColMajor but the CG sovler likes RowMajor
 //     Also, to change everything to float / double just change the scalar type here
 
@@ -82,7 +83,7 @@ public:
 * A struct of arrays representation used to hold vertex positions
 * in three vectors that can interface with the eigen library
 */
-class FSOAPositions : public TVector3Arrays<double>
+class FSOAPositions : public UE::Geometry::TVector3Arrays<double>
 {
 public:
 	typedef typename FSparseMatrixD::Scalar ScalarType;
@@ -90,7 +91,7 @@ public:
 	typedef Eigen::Map<RealVectorType> VectorType;
 	typedef Eigen::Map<const RealVectorType> ConstVectorType;
 
-	FSOAPositions(int32 Size) : TVector3Arrays<double>(Size)
+	FSOAPositions(int32 Size) : UE::Geometry::TVector3Arrays<double>(Size)
 	{ }
 
 	FSOAPositions()
@@ -108,3 +109,5 @@ public:
 	}
 
 };
+
+
diff --git a/Engine/Plugins/Experimental/GeometryProcessing/Source/DynamicMesh/Private/Solvers/Internal/LaplacianOperators.h b/Engine/Plugins/Experimental/GeometryProcessing/Source/DynamicMesh/Private/Solvers/Internal/LaplacianOperators.h
--- a/Engine/Plugins/Experimental/GeometryProcessing/Source/DynamicMesh/Private/Solvers/Internal/LaplacianOperators.h
+++ b/Engine/Plugins/Experimental/GeometryProcessing/Source/DynamicMesh/Private/Solvers/Internal/LaplacianOperators.h
@@ -9,6 +9,8 @@
 
 #include "FSparseMatrixD.h"
 
+// this is a private header
+using namespace UE::Geometry;
 
 /**
 * Construct a sparse matrix representation of the requested "Laplacian" type operator.
diff --git a/Engine/Plugins/Experimental/GeometryProcessing/Source/DynamicMesh/Private/Solvers/Internal/MeshDiffusionIntegrator.h b/Engine/Plugins/Experimental/GeometryProcessing/Source/DynamicMesh/Private/Solvers/Internal/MeshDiffusionIntegrator.h
--- a/Engine/Plugins/Experimental/GeometryProcessing/Source/DynamicMesh/Private/Solvers/Internal/MeshDiffusionIntegrator.h
+++ b/Engine/Plugins/Experimental/GeometryProcessing/Source/DynamicMesh/Private/Solvers/Internal/MeshDiffusionIntegrator.h
@@ -9,6 +9,12 @@
 
 #include "MatrixSolver.h"
 
+namespace UE
+{
+namespace Geometry
+{
+
+
 /**
  * FMeshDiffusionIntegrator solves diffusion problems over vertices of a 3D triangle mesh.
  *
@@ -71,3 +77,6 @@ protected:
 	int32                               Id; // double buffer id
 };
 
+
+} // end namespace UE::Geometry
+} // end namespace UE
\ No newline at end of file
diff --git a/Engine/Plugins/Experimental/GeometryProcessing/Source/DynamicMesh/Private/Solvers/Internal/MeshDiffusionSmoothing.h b/Engine/Plugins/Experimental/GeometryProcessing/Source/DynamicMesh/Private/Solvers/Internal/MeshDiffusionSmoothing.h
--- a/Engine/Plugins/Experimental/GeometryProcessing/Source/DynamicMesh/Private/Solvers/Internal/MeshDiffusionSmoothing.h
+++ b/Engine/Plugins/Experimental/GeometryProcessing/Source/DynamicMesh/Private/Solvers/Internal/MeshDiffusionSmoothing.h
@@ -5,6 +5,10 @@
 #include "MeshDiffusionIntegrator.h"
 #include "LaplacianOperators.h"
 
+namespace UE
+{
+namespace Geometry
+{
 
 /**
  * Solve mesh diffusion using Laplacian matrix of various types
@@ -88,3 +92,7 @@ protected:
 	}
 };
 
+
+
+} // end namespace UE::Geometry
+} // end namespace UE
\ No newline at end of file
diff --git a/Engine/Plugins/Experimental/GeometryProcessing/Source/DynamicMesh/Private/Solvers/Internal/MeshUVSolver.h b/Engine/Plugins/Experimental/GeometryProcessing/Source/DynamicMesh/Private/Solvers/Internal/MeshUVSolver.h
--- a/Engine/Plugins/Experimental/GeometryProcessing/Source/DynamicMesh/Private/Solvers/Internal/MeshUVSolver.h
+++ b/Engine/Plugins/Experimental/GeometryProcessing/Source/DynamicMesh/Private/Solvers/Internal/MeshUVSolver.h
@@ -11,6 +11,11 @@
 
 #include "MatrixBase.h"
 
+namespace UE
+{
+namespace Geometry
+{
+
 enum class EUVSolveType
 {
 	NaturalConformal
@@ -62,4 +67,7 @@ protected:
 
 	// Used to map between VtxId and vertex Index in linear vector..
 	FVertexLinearization  VtxLinearization;
-};
\ No newline at end of file
+};
+
+} // end namespace UE::Geometry
+} // end namespace UE
\ No newline at end of file
diff --git a/Engine/Plugins/Experimental/GeometryProcessing/Source/DynamicMesh/Public/DynamicAttribute.h b/Engine/Plugins/Experimental/GeometryProcessing/Source/DynamicMesh/Public/DynamicAttribute.h
--- a/Engine/Plugins/Experimental/GeometryProcessing/Source/DynamicMesh/Public/DynamicAttribute.h
+++ b/Engine/Plugins/Experimental/GeometryProcessing/Source/DynamicMesh/Public/DynamicAttribute.h
@@ -5,6 +5,13 @@
 #include "DynamicMesh3.h"
 #include "MeshIndexMappings.h"
 
+namespace UE
+{
+namespace Geometry
+{
+
+struct FMeshIndexMappings;
+
 template<typename ParentType>
 class TDynamicAttributeBase;
 
@@ -340,3 +347,6 @@ public:
 };
 
 typedef TDynamicAttributeSetBase<FDynamicMesh3> FDynamicMeshAttributeSetBase;
+
+} // end namespace UE::Geometry
+} // end namespace UE
\ No newline at end of file
diff --git a/Engine/Plugins/Experimental/GeometryProcessing/Source/DynamicMesh/Public/DynamicMesh3.h b/Engine/Plugins/Experimental/GeometryProcessing/Source/DynamicMesh/Public/DynamicMesh3.h
--- a/Engine/Plugins/Experimental/GeometryProcessing/Source/DynamicMesh/Public/DynamicMesh3.h
+++ b/Engine/Plugins/Experimental/GeometryProcessing/Source/DynamicMesh/Public/DynamicMesh3.h
@@ -21,8 +21,13 @@
 #include "VectorTypes.h"
 #include "VectorUtil.h"
 
-class FDynamicMeshAttributeSet;
+namespace UE
+{
+namespace Geometry
+{
+
 class FMeshShapeGenerator;
+class FDynamicMeshAttributeSet;
 
 enum class EMeshComponents : uint8
 {
@@ -1417,3 +1422,8 @@ protected:
 		}
 	}
 };
+
+
+
+} // end namespace UE::Geometry
+} // end namespace UE
diff --git a/Engine/Plugins/Experimental/GeometryProcessing/Source/DynamicMesh/Public/DynamicMeshAABBTree3.h b/Engine/Plugins/Experimental/GeometryProcessing/Source/DynamicMesh/Public/DynamicMeshAABBTree3.h
--- a/Engine/Plugins/Experimental/GeometryProcessing/Source/DynamicMesh/Public/DynamicMeshAABBTree3.h
+++ b/Engine/Plugins/Experimental/GeometryProcessing/Source/DynamicMesh/Public/DynamicMeshAABBTree3.h
@@ -5,4 +5,12 @@
 #include "Spatial/MeshAABBTree3.h"
 #include "DynamicMesh3.h"
 
-typedef TMeshAABBTree3<FDynamicMesh3> FDynamicMeshAABBTree3;
\ No newline at end of file
+namespace UE
+{
+namespace Geometry
+{
+
+typedef TMeshAABBTree3<FDynamicMesh3> FDynamicMeshAABBTree3;
+
+} // end namespace UE::Geometry
+} // end namespace UE
diff --git a/Engine/Plugins/Experimental/GeometryProcessing/Source/DynamicMesh/Public/DynamicMeshAttributeSet.h b/Engine/Plugins/Experimental/GeometryProcessing/Source/DynamicMesh/Public/DynamicMeshAttributeSet.h
--- a/Engine/Plugins/Experimental/GeometryProcessing/Source/DynamicMesh/Public/DynamicMeshAttributeSet.h
+++ b/Engine/Plugins/Experimental/GeometryProcessing/Source/DynamicMesh/Public/DynamicMeshAttributeSet.h
@@ -10,6 +10,11 @@
 #include "InfoTypes.h"
 #include "Containers/IndirectArray.h"
 
+namespace UE
+{
+namespace Geometry
+{
+
 /** Standard UV overlay type - 2-element float */
 typedef TDynamicMeshVectorOverlay<float, 2, FVector2f> FDynamicMeshUVOverlay;
 /** Standard Normal overlay type - 3-element float */
@@ -392,3 +397,8 @@ protected:
 	}
 };
 
+
+
+} // end namespace UE::Geometry
+} // end namespace UE
+
diff --git a/Engine/Plugins/Experimental/GeometryProcessing/Source/DynamicMesh/Public/DynamicMeshChangeTracker.h b/Engine/Plugins/Experimental/GeometryProcessing/Source/DynamicMesh/Public/DynamicMeshChangeTracker.h
--- a/Engine/Plugins/Experimental/GeometryProcessing/Source/DynamicMesh/Public/DynamicMeshChangeTracker.h
+++ b/Engine/Plugins/Experimental/GeometryProcessing/Source/DynamicMesh/Public/DynamicMeshChangeTracker.h
@@ -29,7 +29,10 @@
 // 
 // 
 
-
+namespace UE
+{
+namespace Geometry
+{
 
 
 /**
@@ -398,3 +401,7 @@ void FDynamicMeshChangeTracker::SaveVertexOneRingTriangles(EnumerableType Vertex
 	}
 }
 
+
+
+} // end namespace UE::Geometry
+} // end namespace UE
\ No newline at end of file
diff --git a/Engine/Plugins/Experimental/GeometryProcessing/Source/DynamicMesh/Public/DynamicMeshEditor.h b/Engine/Plugins/Experimental/GeometryProcessing/Source/DynamicMesh/Public/DynamicMeshEditor.h
--- a/Engine/Plugins/Experimental/GeometryProcessing/Source/DynamicMesh/Public/DynamicMeshEditor.h
+++ b/Engine/Plugins/Experimental/GeometryProcessing/Source/DynamicMesh/Public/DynamicMeshEditor.h
@@ -10,6 +10,11 @@
 #include "Util/SparseIndexCollectionTypes.h"
 #include "MeshIndexMappings.h"
 
+namespace UE
+{
+namespace Geometry
+{
+
 struct FDynamicSubmesh3;
 
 /**
@@ -490,3 +495,5 @@ public:
 };
 
 
+} // end namespace UE::Geometry
+} // end namespace UE
diff --git a/Engine/Plugins/Experimental/GeometryProcessing/Source/DynamicMesh/Public/DynamicMeshOctree3.h b/Engine/Plugins/Experimental/GeometryProcessing/Source/DynamicMesh/Public/DynamicMeshOctree3.h
--- a/Engine/Plugins/Experimental/GeometryProcessing/Source/DynamicMesh/Public/DynamicMeshOctree3.h
+++ b/Engine/Plugins/Experimental/GeometryProcessing/Source/DynamicMesh/Public/DynamicMeshOctree3.h
@@ -9,6 +9,12 @@
 #include "Async/ParallelFor.h"
 
 
+namespace UE
+{
+namespace Geometry
+{
+
+
 /**
  * FDynamicMeshOctree3 is an extension of FSparseDynamicOctree3 for the triangles of a FDynamicMesh3 instance.
  * This extension does several things:
@@ -442,4 +448,9 @@ public:
 	}
 
 
-};
\ No newline at end of file
+};
+
+
+
+} // end namespace UE::Geometry
+} // end namespace UE
diff --git a/Engine/Plugins/Experimental/GeometryProcessing/Source/DynamicMesh/Public/DynamicMeshOverlay.h b/Engine/Plugins/Experimental/GeometryProcessing/Source/DynamicMesh/Public/DynamicMeshOverlay.h
--- a/Engine/Plugins/Experimental/GeometryProcessing/Source/DynamicMesh/Public/DynamicMeshOverlay.h
+++ b/Engine/Plugins/Experimental/GeometryProcessing/Source/DynamicMesh/Public/DynamicMeshOverlay.h
@@ -11,6 +11,11 @@
 #include "InfoTypes.h"
 #include "IndexTypes.h"
 
+namespace UE
+{
+namespace Geometry
+{
+
 class FDynamicMesh3;
 
 /**
@@ -547,3 +552,5 @@ public:
 };
 
 
+} // end namespace UE::Geometry
+} // end namespace UE
\ No newline at end of file
diff --git a/Engine/Plugins/Experimental/GeometryProcessing/Source/DynamicMesh/Public/DynamicPointSet3.h b/Engine/Plugins/Experimental/GeometryProcessing/Source/DynamicMesh/Public/DynamicPointSet3.h
--- a/Engine/Plugins/Experimental/GeometryProcessing/Source/DynamicMesh/Public/DynamicPointSet3.h
+++ b/Engine/Plugins/Experimental/GeometryProcessing/Source/DynamicMesh/Public/DynamicPointSet3.h
@@ -15,6 +15,10 @@
 
 class FDynamicAttributeSetBase;
 
+namespace UE
+{
+namespace Geometry
+{
 
 /**
  * TDynamicPointSet3 implements a dynamic 3D point set, templated on real-value type (float or double).
@@ -325,3 +329,5 @@ typedef TDynamicPointSet3<float> FDynamicPointSet3f;
 typedef TDynamicPointSet3<double> FDynamicPointSet3d;
 
 
+} // end namespace UE::Geometry
+} // end namespace UE
\ No newline at end of file
diff --git a/Engine/Plugins/Experimental/GeometryProcessing/Source/DynamicMesh/Public/DynamicSubmesh3.h b/Engine/Plugins/Experimental/GeometryProcessing/Source/DynamicMesh/Public/DynamicSubmesh3.h
--- a/Engine/Plugins/Experimental/GeometryProcessing/Source/DynamicMesh/Public/DynamicSubmesh3.h
+++ b/Engine/Plugins/Experimental/GeometryProcessing/Source/DynamicMesh/Public/DynamicSubmesh3.h
@@ -8,6 +8,11 @@
 #include "DynamicMeshEditor.h"
 #include "Util/SparseIndexCollectionTypes.h"
 
+namespace UE
+{
+namespace Geometry
+{
+
 struct DYNAMICMESH_API FDynamicSubmesh3
 {
 protected:
@@ -260,3 +265,7 @@ protected:
 	}
 };
 
+
+
+} // end namespace UE::Geometry
+} // end namespace UE
\ No newline at end of file
diff --git a/Engine/Plugins/Experimental/GeometryProcessing/Source/DynamicMesh/Public/DynamicVertexAttribute.h b/Engine/Plugins/Experimental/GeometryProcessing/Source/DynamicMesh/Public/DynamicVertexAttribute.h
--- a/Engine/Plugins/Experimental/GeometryProcessing/Source/DynamicMesh/Public/DynamicVertexAttribute.h
+++ b/Engine/Plugins/Experimental/GeometryProcessing/Source/DynamicMesh/Public/DynamicVertexAttribute.h
@@ -5,6 +5,10 @@
 #include "DynamicMesh3.h"
 #include "DynamicAttribute.h"
 
+namespace UE
+{
+namespace Geometry
+{
 
 template<typename AttribValueType, int AttribDimension, typename ParentType>
 class TDynamicVertexAttribute;
@@ -438,4 +442,5 @@ bool TDynamicVertexAttributeChange<AttribValueType, AttribDimension, ParentType>
 }
 
 
-
+} // end namespace UE::Geometry
+} // end namespace UE
diff --git a/Engine/Plugins/Experimental/GeometryProcessing/Source/DynamicMesh/Public/DynamicVerticesOctree3.h b/Engine/Plugins/Experimental/GeometryProcessing/Source/DynamicMesh/Public/DynamicVerticesOctree3.h
--- a/Engine/Plugins/Experimental/GeometryProcessing/Source/DynamicMesh/Public/DynamicVerticesOctree3.h
+++ b/Engine/Plugins/Experimental/GeometryProcessing/Source/DynamicMesh/Public/DynamicVerticesOctree3.h
@@ -8,6 +8,10 @@
 #include "Intersection/IntersectionUtil.h"
 #include "DynamicMesh3.h"
 
+namespace UE
+{
+namespace Geometry
+{
 
 /**
  * TDynamicVerticesOctree3 is an extension of FSparseDynamicPointOctree3 for the vertices of 
@@ -219,4 +223,9 @@ public:
 	}
 
 
-};
\ No newline at end of file
+};
+
+
+
+} // end namespace UE::Geometry
+} // end namespace UE
diff --git a/Engine/Plugins/Experimental/GeometryProcessing/Source/DynamicMesh/Public/EdgeLoop.h b/Engine/Plugins/Experimental/GeometryProcessing/Source/DynamicMesh/Public/EdgeLoop.h
--- a/Engine/Plugins/Experimental/GeometryProcessing/Source/DynamicMesh/Public/EdgeLoop.h
+++ b/Engine/Plugins/Experimental/GeometryProcessing/Source/DynamicMesh/Public/EdgeLoop.h
@@ -7,6 +7,10 @@
 #include "DynamicMesh3.h"
 #include "Util/IndexUtil.h"
 
+namespace UE
+{
+namespace Geometry
+{
 
 /**
  * Sequential lists of vertices/edges in a mesh that form a closed loop
@@ -224,4 +228,8 @@ public:
 	 */
 	static void VertexLoopToEdgeLoop(const FDynamicMesh3* Mesh, const TArray<int>& VertexLoop, TArray<int>& OutEdgeLoop);
 
-};
\ No newline at end of file
+};
+
+
+} // end namespace UE::Geometry
+} // end namespace UE
diff --git a/Engine/Plugins/Experimental/GeometryProcessing/Source/DynamicMesh/Public/EdgeSpan.h b/Engine/Plugins/Experimental/GeometryProcessing/Source/DynamicMesh/Public/EdgeSpan.h
--- a/Engine/Plugins/Experimental/GeometryProcessing/Source/DynamicMesh/Public/EdgeSpan.h
+++ b/Engine/Plugins/Experimental/GeometryProcessing/Source/DynamicMesh/Public/EdgeSpan.h
@@ -8,6 +8,10 @@
 #include "Util/IndexUtil.h"
 #include "Polyline3.h"
 
+namespace UE
+{
+namespace Geometry
+{
 
 /**
  * Sequential lists of vertices/edges in a mesh that is *not* closed.
@@ -204,4 +208,8 @@ public:
 	 */
 	static void VertexSpanToEdgeSpan(const FDynamicMesh3* Mesh, const TArray<int>& VertexSpan, TArray<int>& OutEdgeSpan);
 
-};
\ No newline at end of file
+};
+
+
+} // end namespace UE::Geometry
+} // end namespace UE
diff --git a/Engine/Plugins/Experimental/GeometryProcessing/Source/DynamicMesh/Public/GroupTopology.h b/Engine/Plugins/Experimental/GeometryProcessing/Source/DynamicMesh/Public/GroupTopology.h
--- a/Engine/Plugins/Experimental/GeometryProcessing/Source/DynamicMesh/Public/GroupTopology.h
+++ b/Engine/Plugins/Experimental/GeometryProcessing/Source/DynamicMesh/Public/GroupTopology.h
@@ -8,10 +8,13 @@
 #include "Util/IndexUtil.h"
 #include "Containers/BitArray.h"
 
+namespace UE
+{
+namespace Geometry
+{
 
 class FGroupTopology;
 
-
 /**
  * FGroupTopologySelection represents a set of selected elements of a FGroupTopology
  */
@@ -450,4 +453,9 @@ public:
 	{
 		return TriangleID;
 	}
-};
\ No newline at end of file
+};
+
+
+} // end namespace UE::Geometry
+} // end namespace UE
+
diff --git a/Engine/Plugins/Experimental/GeometryProcessing/Source/DynamicMesh/Public/InfoTypes.h b/Engine/Plugins/Experimental/GeometryProcessing/Source/DynamicMesh/Public/InfoTypes.h
--- a/Engine/Plugins/Experimental/GeometryProcessing/Source/DynamicMesh/Public/InfoTypes.h
+++ b/Engine/Plugins/Experimental/GeometryProcessing/Source/DynamicMesh/Public/InfoTypes.h
@@ -5,6 +5,11 @@
 #include "IndexTypes.h"
 #include "VectorTypes.h"
 
+namespace UE
+{
+namespace Geometry
+{
+
 /**
  * FVertexInfo stores information about vertex attributes - position, normal, color, UV
  */
@@ -27,8 +32,16 @@ struct FVertexInfo
 		: Position{ PositionIn }, Normal{ NormalIn }, Color{ ColorIn }, UV{ UVIn }, bHaveN{ true }, bHaveC{true}, bHaveUV{true}{}
 };
 
+
+} // end namespace UE::Geometry
+} // end namespace UE
+
+
+
 namespace DynamicMeshInfo
 {
+using namespace UE::Geometry;
+using FVector3d = UE::Geometry::FVector3d;
 
 /** Information about the mesh elements created by a call to SplitEdge() */
 struct FEdgeSplitInfo
diff --git a/Engine/Plugins/Experimental/GeometryProcessing/Source/DynamicMesh/Public/MeshCurvature.h b/Engine/Plugins/Experimental/GeometryProcessing/Source/DynamicMesh/Public/MeshCurvature.h
--- a/Engine/Plugins/Experimental/GeometryProcessing/Source/DynamicMesh/Public/MeshCurvature.h
+++ b/Engine/Plugins/Experimental/GeometryProcessing/Source/DynamicMesh/Public/MeshCurvature.h
@@ -12,6 +12,8 @@ namespace UE
 	 */
 	namespace MeshCurvature
 	{
+		using namespace UE::Geometry;
+
 		/**
 		 * Calculate the Discrete Mean-Curvature Normal at a vertex as defined by discrete differential geometry.
 		 * Based on Eq 8 from "Discrete Differential-Geometry Operators for Triangulated 2-Manifolds", Meyer et al 2002
@@ -47,6 +49,10 @@ namespace UE
 }
 
 
+namespace UE
+{
+namespace Geometry
+{
 
 
 /**
@@ -76,4 +82,7 @@ public:
 	const FVertexCurvature& operator[](int32 Index) const { return Curvatures[Index]; }
 
 	void BuildAll(const FDynamicMesh3& Mesh);
-};
\ No newline at end of file
+};
+
+} // end namespace UE::Geometry
+} // end namespace UE
\ No newline at end of file
diff --git a/Engine/Plugins/Experimental/GeometryProcessing/Source/DynamicMesh/Public/MeshIndexUtil.h b/Engine/Plugins/Experimental/GeometryProcessing/Source/DynamicMesh/Public/MeshIndexUtil.h
--- a/Engine/Plugins/Experimental/GeometryProcessing/Source/DynamicMesh/Public/MeshIndexUtil.h
+++ b/Engine/Plugins/Experimental/GeometryProcessing/Source/DynamicMesh/Public/MeshIndexUtil.h
@@ -9,6 +9,8 @@
  */
 namespace MeshIndexUtil
 {
+	using namespace UE::Geometry;
+
 	//
 	// Do not add any more functions to this namespace. Use UE::MeshIndexUtil instead
 	//
@@ -36,6 +38,8 @@ namespace UE
 {
 	namespace MeshIndexUtil
 	{
+		using namespace UE::Geometry;
+
 		/**
 		 * Find list of unique vertices that are contained in one or more triangles
 		 * @param Mesh input mesh
diff --git a/Engine/Plugins/Experimental/GeometryProcessing/Source/DynamicMesh/Public/MeshNormals.h b/Engine/Plugins/Experimental/GeometryProcessing/Source/DynamicMesh/Public/MeshNormals.h
--- a/Engine/Plugins/Experimental/GeometryProcessing/Source/DynamicMesh/Public/MeshNormals.h
+++ b/Engine/Plugins/Experimental/GeometryProcessing/Source/DynamicMesh/Public/MeshNormals.h
@@ -7,6 +7,12 @@
 #include "DynamicMesh3.h"
 #include "DynamicMeshAttributeSet.h"
 
+namespace UE
+{
+namespace Geometry
+{
+
+
 /**
  * FMeshNormals is a utility class that can calculate and store various types of
  * normal vectors for a FDynamicMesh. 
@@ -192,4 +198,8 @@ protected:
 		return TriNormalWeights;
 	}
 
-};
\ No newline at end of file
+};
+
+
+} // end namespace UE::Geometry
+} // end namespace UE
\ No newline at end of file
diff --git a/Engine/Plugins/Experimental/GeometryProcessing/Source/DynamicMesh/Public/MeshSimplification.h b/Engine/Plugins/Experimental/GeometryProcessing/Source/DynamicMesh/Public/MeshSimplification.h
--- a/Engine/Plugins/Experimental/GeometryProcessing/Source/DynamicMesh/Public/MeshSimplification.h
+++ b/Engine/Plugins/Experimental/GeometryProcessing/Source/DynamicMesh/Public/MeshSimplification.h
@@ -11,6 +11,11 @@
 #include "DynamicMeshAttributeSet.h"
 
 
+namespace UE
+{
+namespace Geometry
+{
+
 enum class ESimplificationResult
 {
 	Ok_Collapsed = 0,
@@ -401,4 +406,8 @@ protected:
 // The simplifier
 typedef TMeshSimplification< FAttrBasedQuadricErrord >  FAttrMeshSimplification;
 typedef TMeshSimplification < FVolPresQuadricErrord >   FVolPresMeshSimplification;
-typedef TMeshSimplification< FQuadricErrord >           FQEMSimplification;
\ No newline at end of file
+typedef TMeshSimplification< FQuadricErrord >           FQEMSimplification;
+
+
+} // end namespace UE::Geometry
+} // end namespace UE
\ No newline at end of file
diff --git a/Engine/Plugins/Experimental/GeometryProcessing/Source/DynamicMesh/Public/MeshTangents.h b/Engine/Plugins/Experimental/GeometryProcessing/Source/DynamicMesh/Public/MeshTangents.h
--- a/Engine/Plugins/Experimental/GeometryProcessing/Source/DynamicMesh/Public/MeshTangents.h
+++ b/Engine/Plugins/Experimental/GeometryProcessing/Source/DynamicMesh/Public/MeshTangents.h
@@ -6,6 +6,12 @@
 #include "DynamicMeshAttributeSet.h"
 
 
+namespace UE
+{
+namespace Geometry
+{
+
+
 /**
  * Options used by TMeshTangents for tangents computation
  */
@@ -208,4 +214,8 @@ void TMeshTangents<RealType>::CopyTriVertexTangents(const TMeshTangents<OtherRea
 		Bitangents[k] = FVector3<RealType>(OtherBitangents[k]);
 	}
 	AllDegenerateTris = OtherMeshTangents.GetDegenerateTris();
-}
\ No newline at end of file
+}
+
+
+} // end namespace UE::Geometry
+} // end namespace UE
\ No newline at end of file
diff --git a/Engine/Plugins/Experimental/GeometryProcessing/Source/DynamicMesh/Public/MeshTransforms.h b/Engine/Plugins/Experimental/GeometryProcessing/Source/DynamicMesh/Public/MeshTransforms.h
--- a/Engine/Plugins/Experimental/GeometryProcessing/Source/DynamicMesh/Public/MeshTransforms.h
+++ b/Engine/Plugins/Experimental/GeometryProcessing/Source/DynamicMesh/Public/MeshTransforms.h
@@ -12,7 +12,10 @@
  */
 namespace MeshTransforms
 {
-	
+	using namespace UE::Geometry;
+	using FVector3d = UE::Geometry::FVector3d;
+	using FVector3f = UE::Geometry::FVector3f;
+
 	/**
 	 * Apply Translation to vertex positions of Mesh. Does not modify any other attributes.
 	 */
diff --git a/Engine/Plugins/Experimental/GeometryProcessing/Source/DynamicMesh/Public/MeshWeights.h b/Engine/Plugins/Experimental/GeometryProcessing/Source/DynamicMesh/Public/MeshWeights.h
--- a/Engine/Plugins/Experimental/GeometryProcessing/Source/DynamicMesh/Public/MeshWeights.h
+++ b/Engine/Plugins/Experimental/GeometryProcessing/Source/DynamicMesh/Public/MeshWeights.h
@@ -6,6 +6,11 @@
 
 #include "DynamicMesh3.h"
 
+namespace UE
+{
+namespace Geometry
+{
+
 /**
  * FMeshWeights implements various techniques for computing local weights of a mesh,
  * for example one-ring weights like Cotangent or Mean-Value.
@@ -93,4 +98,8 @@ public:
 
 protected:
 	FMeshWeights() = delete;
-};
\ No newline at end of file
+};
+
+
+} // end namespace UE::Geometry
+} // end namespace UE
\ No newline at end of file
diff --git a/Engine/Plugins/Experimental/GeometryProcessing/Source/DynamicMesh/Public/NormalFlowRemesher.h b/Engine/Plugins/Experimental/GeometryProcessing/Source/DynamicMesh/Public/NormalFlowRemesher.h
--- a/Engine/Plugins/Experimental/GeometryProcessing/Source/DynamicMesh/Public/NormalFlowRemesher.h
+++ b/Engine/Plugins/Experimental/GeometryProcessing/Source/DynamicMesh/Public/NormalFlowRemesher.h
@@ -4,13 +4,17 @@
 
 #include "QueueRemesher.h"
 
+namespace UE
+{
+namespace Geometry
+{
+
 /**
  * Remeshing with "face aligned projection". This approach to projection attempts to preserve sharp features in the 
  * mesh by aligning triangle normals with normals from the original mesh during the projection pass.
  * This class also performs a pass of edge flips aimed at further aligning triangle normals to their originals. This
  * can be helpful in fixing the occasional single "bad edge" along a sequence of feature edges.
  */
-
 class DYNAMICMESH_API FNormalFlowRemesher : public FQueueRemesher
 {
 public:
@@ -83,3 +87,7 @@ protected:
 	/** Flip all edges for which EdgeFlipWouldReduceNormalError returns true */
 	void TrackedEdgeFlipPass();
 };
+
+
+} // end namespace UE::Geometry
+} // end namespace UE
\ No newline at end of file
diff --git a/Engine/Plugins/Experimental/GeometryProcessing/Source/DynamicMesh/Public/Operations/EmbedSurfacePath.h b/Engine/Plugins/Experimental/GeometryProcessing/Source/DynamicMesh/Public/Operations/EmbedSurfacePath.h
--- a/Engine/Plugins/Experimental/GeometryProcessing/Source/DynamicMesh/Public/Operations/EmbedSurfacePath.h
+++ b/Engine/Plugins/Experimental/GeometryProcessing/Source/DynamicMesh/Public/Operations/EmbedSurfacePath.h
@@ -12,6 +12,12 @@
 #include "Curve/GeneralPolygon2.h"
 #include "Selections/MeshFaceSelection.h"
 
+
+namespace UE
+{
+namespace Geometry
+{
+
 class FDynamicMesh3;
 
 enum class ESurfacePointType
@@ -46,10 +52,10 @@ struct DYNAMICMESH_API FMeshSurfacePoint
 /**
  * Walk the surface of an FDynamicMesh to try find a planar path connecting two points.  Paths include every vertex and edge they need to cross.  Greedy algorithm will only return one path if there are multiple.
  */
-bool DYNAMICMESH_API WalkMeshPlanar(
-	const FDynamicMesh3* Mesh, int StartTri, int EndVertID, FVector3d StartPt, int EndTri, FVector3d EndPt, FVector3d WalkPlaneNormal, 
-	TFunction<FVector3d(const FDynamicMesh3*, int)> VertexToPosnFn, bool bAllowBackwardsSearch, double AcceptEndPtOutsideDist,
-	double PtOnPlaneThreshold, TArray<TPair<FMeshSurfacePoint, int>>& WalkedPath, double BackwardsTolerance = FMathd::ZeroTolerance * 10);
+//bool DYNAMICMESH_API WalkMeshPlanar(
+//	const FDynamicMesh3* Mesh, int StartTri, int EndVertID, FVector3d StartPt, int EndTri, FVector3d EndPt, FVector3d WalkPlaneNormal, 
+//	TFunction<FVector3d(const FDynamicMesh3*, int)> VertexToPosnFn, bool bAllowBackwardsSearch, double AcceptEndPtOutsideDist,
+//	double PtOnPlaneThreshold, TArray<TPair<FMeshSurfacePoint, int>>& WalkedPath, double BackwardsTolerance = FMathd::ZeroTolerance * 10);
 
 
 /**
@@ -149,3 +155,7 @@ bool DYNAMICMESH_API EmbedProjectedPath(FDynamicMesh3* Mesh, int StartTriID, FFr
  * Embed multiple 2D paths into a mesh by projection, starting the walks from the given triangles.  Optionally select the triangles inside the paths.
  */
 bool DYNAMICMESH_API EmbedProjectedPaths(FDynamicMesh3* Mesh, const TArrayView<const int> StartTriIDs, FFrame3d Frame, const TArrayView<const TArray<FVector2d>> AllPaths, TArray<TArray<int>>& OutAllPathVertices, TArray<TArray<int>>& OutAllVertexCorrespondence, bool bClosePaths, FMeshFaceSelection* EnclosedFaces, double PtSnapVertexOrEdgeThresholdSq = FMathf::ZeroTolerance*100);
+
+
+} // end namespace UE::Geometry
+} // end namespace UE
\ No newline at end of file
diff --git a/Engine/Plugins/Experimental/GeometryProcessing/Source/DynamicMesh/Public/Operations/ExtrudeMesh.h b/Engine/Plugins/Experimental/GeometryProcessing/Source/DynamicMesh/Public/Operations/ExtrudeMesh.h
--- a/Engine/Plugins/Experimental/GeometryProcessing/Source/DynamicMesh/Public/Operations/ExtrudeMesh.h
+++ b/Engine/Plugins/Experimental/GeometryProcessing/Source/DynamicMesh/Public/Operations/ExtrudeMesh.h
@@ -10,6 +10,11 @@
 #include "MeshRegionBoundaryLoops.h"
 
 
+namespace UE
+{
+namespace Geometry
+{
+
 class FDynamicMesh3;
 class FMeshNormals;
 
@@ -114,4 +119,7 @@ public:
 protected:
 
 	bool ApplyExtrude(FExtrusionInfo& Region, FMeshNormals* UseNormals = nullptr);
-};
\ No newline at end of file
+};
+
+} // end namespace UE::Geometry
+} // end namespace UE
diff --git a/Engine/Plugins/Experimental/GeometryProcessing/Source/DynamicMesh/Public/Operations/FFDLattice.h b/Engine/Plugins/Experimental/GeometryProcessing/Source/DynamicMesh/Public/Operations/FFDLattice.h
--- a/Engine/Plugins/Experimental/GeometryProcessing/Source/DynamicMesh/Public/Operations/FFDLattice.h
+++ b/Engine/Plugins/Experimental/GeometryProcessing/Source/DynamicMesh/Public/Operations/FFDLattice.h
@@ -7,9 +7,14 @@
 #include "Containers/StaticArray.h"
 #include "DynamicMeshAttributeSet.h"
 
-class FDynamicMesh3;
 class FProgressCancel;
 
+namespace UE
+{
+namespace Geometry
+{
+
+class FDynamicMesh3;
 
 enum class ELatticeInterpolation : uint8
 {
@@ -124,3 +129,6 @@ protected:
 
 };
 
+
+} // end namespace UE::Geometry
+} // end namespace UE
\ No newline at end of file
diff --git a/Engine/Plugins/Experimental/GeometryProcessing/Source/DynamicMesh/Public/Operations/GroupEdgeInserter.h b/Engine/Plugins/Experimental/GeometryProcessing/Source/DynamicMesh/Public/Operations/GroupEdgeInserter.h
--- a/Engine/Plugins/Experimental/GeometryProcessing/Source/DynamicMesh/Public/Operations/GroupEdgeInserter.h
+++ b/Engine/Plugins/Experimental/GeometryProcessing/Source/DynamicMesh/Public/Operations/GroupEdgeInserter.h
@@ -7,6 +7,12 @@
 #include "GroupTopology.h"
 
 class FProgressCancel;
+
+namespace UE
+{
+namespace Geometry
+{
+
 class FDynamicMesh3;
 class FGroupTopology;
 
@@ -113,4 +119,8 @@ public:
 
 	bool InsertGroupEdge(FGroupEdgeInsertionParams& Params, TSet<int32>* NewEidsOut = nullptr, FProgressCancel* Progress = nullptr);
 
-};
\ No newline at end of file
+};
+
+
+} // end namespace UE::Geometry
+} // end namespace UE
\ No newline at end of file
diff --git a/Engine/Plugins/Experimental/GeometryProcessing/Source/DynamicMesh/Public/Operations/GroupTopologyDeformer.h b/Engine/Plugins/Experimental/GeometryProcessing/Source/DynamicMesh/Public/Operations/GroupTopologyDeformer.h
--- a/Engine/Plugins/Experimental/GeometryProcessing/Source/DynamicMesh/Public/Operations/GroupTopologyDeformer.h
+++ b/Engine/Plugins/Experimental/GeometryProcessing/Source/DynamicMesh/Public/Operations/GroupTopologyDeformer.h
@@ -4,11 +4,16 @@
 
 #include "EdgeSpan.h"
 
+
+namespace UE
+{
+namespace Geometry
+{
+
 class FDynamicMesh3;
 class FGroupTopology;
 
 
-
 /**
  * Basic cache of vertex positions.
  */
@@ -195,7 +200,9 @@ protected:
 	};
 	TArray<FFaceEncoding> FaceEncodings;
 
+};
 
 
 
-};
+} // end namespace UE::Geometry
+} // end namespace UE
\ No newline at end of file
diff --git a/Engine/Plugins/Experimental/GeometryProcessing/Source/DynamicMesh/Public/Operations/HoleFiller.h b/Engine/Plugins/Experimental/GeometryProcessing/Source/DynamicMesh/Public/Operations/HoleFiller.h
--- a/Engine/Plugins/Experimental/GeometryProcessing/Source/DynamicMesh/Public/Operations/HoleFiller.h
+++ b/Engine/Plugins/Experimental/GeometryProcessing/Source/DynamicMesh/Public/Operations/HoleFiller.h
@@ -4,6 +4,12 @@
 
 #include "Containers/Array.h"
 
+namespace UE
+{
+namespace Geometry
+{
+
+
 class IHoleFiller
 {
 public:
@@ -13,3 +19,7 @@ public:
 
 	virtual bool Fill(int32 GroupID = -1) = 0;	
 };
+
+
+} // end namespace UE::Geometry
+} // end namespace UE
\ No newline at end of file
diff --git a/Engine/Plugins/Experimental/GeometryProcessing/Source/DynamicMesh/Public/Operations/InsetMeshRegion.h b/Engine/Plugins/Experimental/GeometryProcessing/Source/DynamicMesh/Public/Operations/InsetMeshRegion.h
--- a/Engine/Plugins/Experimental/GeometryProcessing/Source/DynamicMesh/Public/Operations/InsetMeshRegion.h
+++ b/Engine/Plugins/Experimental/GeometryProcessing/Source/DynamicMesh/Public/Operations/InsetMeshRegion.h
@@ -8,9 +8,15 @@
 #include "MeshRegionBoundaryLoops.h"
 #include "ProjectionTargets.h"
 
+
+namespace UE
+{
+namespace Geometry
+{
+
 class FDynamicMesh3;
-class FDynamicMeshChangeTracker;
 class FMeshNormals;
+class FDynamicMeshChangeTracker;
 
 /**
  * FInsetMeshRegion implements local inset of a mesh region.
@@ -124,4 +130,8 @@ public:
 protected:
 
 	virtual bool ApplyInset(FInsetInfo& Region, FMeshNormals* UseNormals = nullptr);
-};
\ No newline at end of file
+};
+
+
+} // end namespace UE::Geometry
+} // end namespace UE
\ No newline at end of file
diff --git a/Engine/Plugins/Experimental/GeometryProcessing/Source/DynamicMesh/Public/Operations/JoinMeshLoops.h b/Engine/Plugins/Experimental/GeometryProcessing/Source/DynamicMesh/Public/Operations/JoinMeshLoops.h
--- a/Engine/Plugins/Experimental/GeometryProcessing/Source/DynamicMesh/Public/Operations/JoinMeshLoops.h
+++ b/Engine/Plugins/Experimental/GeometryProcessing/Source/DynamicMesh/Public/Operations/JoinMeshLoops.h
@@ -8,6 +8,12 @@
 #include "MeshRegionBoundaryLoops.h"
 
 
+
+namespace UE
+{
+namespace Geometry
+{
+
 class FDynamicMesh3;
 class FMeshNormals;
 
@@ -84,4 +90,7 @@ public:
 protected:
 
 	
-};
\ No newline at end of file
+};
+
+} // end namespace UE::Geometry
+} // end namespace UE
\ No newline at end of file
diff --git a/Engine/Plugins/Experimental/GeometryProcessing/Source/DynamicMesh/Public/Operations/MeshAttributeTransfer.h b/Engine/Plugins/Experimental/GeometryProcessing/Source/DynamicMesh/Public/Operations/MeshAttributeTransfer.h
--- a/Engine/Plugins/Experimental/GeometryProcessing/Source/DynamicMesh/Public/Operations/MeshAttributeTransfer.h
+++ b/Engine/Plugins/Experimental/GeometryProcessing/Source/DynamicMesh/Public/Operations/MeshAttributeTransfer.h
@@ -6,8 +6,12 @@
 #include "VectorTypes.h"
 #include "DynamicMeshAABBTree3.h"
 
-class FDynamicMesh3;
+namespace UE
+{
+namespace Geometry
+{
 
+class FDynamicMesh3;
 
 enum class EMeshAttributeTransferType
 {
@@ -58,4 +62,7 @@ protected:
 
 
 	bool Apply_MaterialID();
-};
\ No newline at end of file
+};
+
+} // end namespace UE::Geometry
+} // end namespace UE
\ No newline at end of file
diff --git a/Engine/Plugins/Experimental/GeometryProcessing/Source/DynamicMesh/Public/Operations/MeshBoolean.h b/Engine/Plugins/Experimental/GeometryProcessing/Source/DynamicMesh/Public/Operations/MeshBoolean.h
--- a/Engine/Plugins/Experimental/GeometryProcessing/Source/DynamicMesh/Public/Operations/MeshBoolean.h
+++ b/Engine/Plugins/Experimental/GeometryProcessing/Source/DynamicMesh/Public/Operations/MeshBoolean.h
@@ -16,7 +16,10 @@
 #include "Util/ProgressCancel.h"
 
 
-
+namespace UE
+{
+namespace Geometry
+{
 
 
 
@@ -205,3 +208,7 @@ public:
 		int32 RemoveV, const FVector3d& RemoveVPos, int32 KeepV, const FVector3d& KeepVPos,
 		const FVector3d& EdgeDir, bool bPreserveUVsForMesh, bool bPreserveVertexUVs, bool bPreserveOverlayUVs, float UVToleranceSq);
 };
+
+
+} // end namespace UE::Geometry
+} // end namespace UE
\ No newline at end of file
diff --git a/Engine/Plugins/Experimental/GeometryProcessing/Source/DynamicMesh/Public/Operations/MeshPlaneCut.h b/Engine/Plugins/Experimental/GeometryProcessing/Source/DynamicMesh/Public/Operations/MeshPlaneCut.h
--- a/Engine/Plugins/Experimental/GeometryProcessing/Source/DynamicMesh/Public/Operations/MeshPlaneCut.h
+++ b/Engine/Plugins/Experimental/GeometryProcessing/Source/DynamicMesh/Public/Operations/MeshPlaneCut.h
@@ -10,9 +10,15 @@
 #include "MeshBoundaryLoops.h"
 #include "Curve/GeneralPolygon2.h"
 
+
+
+namespace UE
+{
+namespace Geometry
+{
+
 class FDynamicMesh3;
-template<typename RealType>
-class TDynamicMeshScalarTriangleAttribute;
+template<typename RealType> class TDynamicMeshScalarTriangleAttribute;
 
 /**
  * Cut the Mesh with the Plane. The *positive* side, ie (p-o).n > 0, is removed.
@@ -157,3 +163,7 @@ protected:
 	// set of vertices lying on plane after calling SplitCrossingEdges
 	TSet<int32> OnCutVertices;
 };
+
+
+} // end namespace UE::Geometry
+} // end namespace UE
\ No newline at end of file
diff --git a/Engine/Plugins/Experimental/GeometryProcessing/Source/DynamicMesh/Public/Operations/OffsetMeshRegion.h b/Engine/Plugins/Experimental/GeometryProcessing/Source/DynamicMesh/Public/Operations/OffsetMeshRegion.h
--- a/Engine/Plugins/Experimental/GeometryProcessing/Source/DynamicMesh/Public/Operations/OffsetMeshRegion.h
+++ b/Engine/Plugins/Experimental/GeometryProcessing/Source/DynamicMesh/Public/Operations/OffsetMeshRegion.h
@@ -8,9 +8,14 @@
 #include "MeshRegionBoundaryLoops.h"
 
 
+namespace UE
+{
+namespace Geometry
+{
+
 class FDynamicMesh3;
-class FDynamicMeshChangeTracker;
 class FMeshNormals;
+class FDynamicMeshChangeTracker;
 
 /**
  * FOffsetMeshRegion implements local extrusion/offset of a mesh region. 
@@ -126,4 +131,7 @@ protected:
 	virtual bool ApplyOffset(FOffsetInfo& Region, FMeshNormals* UseNormals = nullptr);
 
 	virtual bool ApplySolidExtrude(FOffsetInfo& Region, FMeshNormals* UseNormals = nullptr);
-};
\ No newline at end of file
+};
+
+} // end namespace UE::Geometry
+} // end namespace UE
diff --git a/Engine/Plugins/Experimental/GeometryProcessing/Source/DynamicMesh/Public/Operations/PlanarHoleFiller.h b/Engine/Plugins/Experimental/GeometryProcessing/Source/DynamicMesh/Public/Operations/PlanarHoleFiller.h
--- a/Engine/Plugins/Experimental/GeometryProcessing/Source/DynamicMesh/Public/Operations/PlanarHoleFiller.h
+++ b/Engine/Plugins/Experimental/GeometryProcessing/Source/DynamicMesh/Public/Operations/PlanarHoleFiller.h
@@ -14,8 +14,12 @@
 #include "MeshBoundaryLoops.h"
 
 
-class FDynamicMesh3;
+namespace UE
+{
+namespace Geometry
+{
 
+class FDynamicMesh3;
 
 /**
  * Fill a set of boundary loops with planar surfaces.  User must provide the triangulation function.
@@ -57,3 +61,7 @@ public:
 
 	bool Fill(int GroupID = -1) override;
 };
+
+
+} // end namespace UE::Geometry
+} // end namespace UE
\ No newline at end of file
diff --git a/Engine/Plugins/Experimental/GeometryProcessing/Source/DynamicMesh/Public/Operations/RemoveOccludedTriangles.h b/Engine/Plugins/Experimental/GeometryProcessing/Source/DynamicMesh/Public/Operations/RemoveOccludedTriangles.h
--- a/Engine/Plugins/Experimental/GeometryProcessing/Source/DynamicMesh/Public/Operations/RemoveOccludedTriangles.h
+++ b/Engine/Plugins/Experimental/GeometryProcessing/Source/DynamicMesh/Public/Operations/RemoveOccludedTriangles.h
@@ -22,6 +22,11 @@
 
 #include "Util/ProgressCancel.h"
 
+namespace UE
+{
+namespace Geometry
+{
+
 enum class EOcclusionTriangleSampling
 {
 	Centroids,
@@ -35,10 +40,17 @@ enum class EOcclusionCalculationMode
 	SimpleOcclusionTest
 };
 
+
+} // end namespace UE::Geometry
+} // end namespace UE
+
+
 namespace UE
 {
 	namespace MeshAutoRepair
 	{
+		using namespace UE::Geometry;
+
 		/**
 		 * Remove any triangles that are internal to the input mesh
 		 * @param Mesh Input mesh to analyze and remove triangles from
@@ -55,6 +67,11 @@ namespace UE
 	}
 }
 
+namespace UE
+{
+namespace Geometry
+{
+
 /**
  * Remove "occluded" triangles, i.e. triangles on the "inside" of the mesh(es).
  * This is a fuzzy definition, current implementation has a couple of options
@@ -606,3 +623,6 @@ protected:
 
 };
 
+
+} // end namespace UE::Geometry
+} // end namespace UE
\ No newline at end of file
diff --git a/Engine/Plugins/Experimental/GeometryProcessing/Source/DynamicMesh/Public/Operations/SimpleHoleFiller.h b/Engine/Plugins/Experimental/GeometryProcessing/Source/DynamicMesh/Public/Operations/SimpleHoleFiller.h
--- a/Engine/Plugins/Experimental/GeometryProcessing/Source/DynamicMesh/Public/Operations/SimpleHoleFiller.h
+++ b/Engine/Plugins/Experimental/GeometryProcessing/Source/DynamicMesh/Public/Operations/SimpleHoleFiller.h
@@ -10,10 +10,13 @@
 #include "GeometryTypes.h"
 #include "MeshRegionBoundaryLoops.h"
 
+namespace UE
+{
+namespace Geometry
+{
 
 class FDynamicMesh3;
 
-
 /**
  * Fill an EdgeLoop hole with triangles.
  * Supports two fill modes, either a fan connected to a new central vertex, or a triangulation of the boundary polygon
@@ -91,3 +94,7 @@ protected:
 	bool Fill_Fan(int32 NewGroupID);
 	bool Fill_EarClip(int32 NewGroupID);
 };
+
+
+} // end namespace UE::Geometry
+} // end namespace UE
\ No newline at end of file
diff --git a/Engine/Plugins/Experimental/GeometryProcessing/Source/DynamicMesh/Public/Operations/SmoothHoleFiller.h b/Engine/Plugins/Experimental/GeometryProcessing/Source/DynamicMesh/Public/Operations/SmoothHoleFiller.h
--- a/Engine/Plugins/Experimental/GeometryProcessing/Source/DynamicMesh/Public/Operations/SmoothHoleFiller.h
+++ b/Engine/Plugins/Experimental/GeometryProcessing/Source/DynamicMesh/Public/Operations/SmoothHoleFiller.h
@@ -10,11 +10,16 @@
 #include "DynamicMeshAABBTree3.h"
 #include "ProjectionTargets.h"
 
+namespace UE
+{
+namespace Geometry
+{
+
 class FDynamicMesh3;
 class FSubRegionRemesher;
-class FEdgeLoop;
-class FMeshBoundaryLoops;
 class FMeshFaceSelection;
+class FMeshBoundaryLoops;
+class FEdgeLoop;
 
 /**
  * Fill parameters
@@ -92,3 +97,7 @@ protected:
 	void DefaultConfigureRemesher(FSubRegionRemesher& Remesher, bool bConstrainROIBoundary);
 
 };
+
+
+} // end namespace UE::Geometry
+} // end namespace UE
\ No newline at end of file
diff --git a/Engine/Plugins/Experimental/GeometryProcessing/Source/DynamicMesh/Public/Parameterization/DynamicMeshUVEditor.h b/Engine/Plugins/Experimental/GeometryProcessing/Source/DynamicMesh/Public/Parameterization/DynamicMeshUVEditor.h
--- a/Engine/Plugins/Experimental/GeometryProcessing/Source/DynamicMesh/Public/Parameterization/DynamicMeshUVEditor.h
+++ b/Engine/Plugins/Experimental/GeometryProcessing/Source/DynamicMesh/Public/Parameterization/DynamicMeshUVEditor.h
@@ -5,6 +5,10 @@
 #include "DynamicMesh3.h"
 #include "DynamicMeshAttributeSet.h"
 
+namespace UE
+{
+namespace Geometry
+{
 
 
 struct DYNAMICMESH_API FUVEditResult
@@ -88,4 +92,9 @@ public:
 	 * @return true on success
 	 */
 	bool CreateSeamAlongVertexPath(const TArray<int32>& VertexPath, FUVEditResult* Result = nullptr);
-};
\ No newline at end of file
+};
+
+
+} // end namespace UE::Geometry
+} // end namespace UE
+
diff --git a/Engine/Plugins/Experimental/GeometryProcessing/Source/DynamicMesh/Public/ProjectionTargets.h b/Engine/Plugins/Experimental/GeometryProcessing/Source/DynamicMesh/Public/ProjectionTargets.h
--- a/Engine/Plugins/Experimental/GeometryProcessing/Source/DynamicMesh/Public/ProjectionTargets.h
+++ b/Engine/Plugins/Experimental/GeometryProcessing/Source/DynamicMesh/Public/ProjectionTargets.h
@@ -7,6 +7,12 @@
 #include "DynamicMeshAABBTree3.h"
 
 
+namespace UE
+{
+namespace Geometry
+{
+
+
 /**
  * FMeshProjectionTarget provides an IProjectionTarget interface to a FDynamicMesh + FDynamicMeshAABBTree3
  * Use to project points to mesh surface.
@@ -43,4 +49,8 @@ public:
 	virtual FVector3d Project(const FVector3d& Point, FVector3d& ProjectNormalOut, int Identifier = -1) override;
 
 
-};
\ No newline at end of file
+};
+
+
+} // end namespace UE::Geometry
+} // end namespace UE
\ No newline at end of file
diff --git a/Engine/Plugins/Experimental/GeometryProcessing/Source/DynamicMesh/Public/QueueRemesher.h b/Engine/Plugins/Experimental/GeometryProcessing/Source/DynamicMesh/Public/QueueRemesher.h
--- a/Engine/Plugins/Experimental/GeometryProcessing/Source/DynamicMesh/Public/QueueRemesher.h
+++ b/Engine/Plugins/Experimental/GeometryProcessing/Source/DynamicMesh/Public/QueueRemesher.h
@@ -5,6 +5,12 @@
 #include "Remesher.h"
 
 
+namespace UE
+{
+namespace Geometry
+{
+
+
 /**
  *
  * Extension to Remesher that is smarter about which edges/vertices to touch:
@@ -15,7 +21,6 @@
  *  - TrackedFullProjectionPass() projects all vertices but only adds to queue if edge changes enough
  * 
  */
-
 class DYNAMICMESH_API FQueueRemesher : public FRemesher
 {
 
@@ -151,3 +156,8 @@ protected:
 
 
 };
+
+
+
+} // end namespace UE::Geometry
+} // end namespace UE
\ No newline at end of file
diff --git a/Engine/Plugins/Experimental/GeometryProcessing/Source/DynamicMesh/Public/Sampling/MeshImageBaker.h b/Engine/Plugins/Experimental/GeometryProcessing/Source/DynamicMesh/Public/Sampling/MeshImageBaker.h
--- a/Engine/Plugins/Experimental/GeometryProcessing/Source/DynamicMesh/Public/Sampling/MeshImageBaker.h
+++ b/Engine/Plugins/Experimental/GeometryProcessing/Source/DynamicMesh/Public/Sampling/MeshImageBaker.h
@@ -4,6 +4,10 @@
 
 #include "Sampling/MeshImageBakingCache.h"
 
+namespace UE
+{
+namespace Geometry
+{
 
 class DYNAMICMESH_API FMeshImageBaker
 {
@@ -28,4 +32,7 @@ protected:
 
 	const FMeshImageBakingCache* ImageBakeCache;
 
-};
\ No newline at end of file
+};
+
+} // end namespace UE::Geometry
+} // end namespace UE
diff --git a/Engine/Plugins/Experimental/GeometryProcessing/Source/DynamicMesh/Public/Sampling/MeshOcclusionMapBaker.h b/Engine/Plugins/Experimental/GeometryProcessing/Source/DynamicMesh/Public/Sampling/MeshOcclusionMapBaker.h
--- a/Engine/Plugins/Experimental/GeometryProcessing/Source/DynamicMesh/Public/Sampling/MeshOcclusionMapBaker.h
+++ b/Engine/Plugins/Experimental/GeometryProcessing/Source/DynamicMesh/Public/Sampling/MeshOcclusionMapBaker.h
@@ -7,6 +7,10 @@
 #include "Sampling/MeshImageBakingCache.h"
 #include "Image/ImageBuilder.h"
 
+namespace UE
+{
+namespace Geometry
+{
 
 class DYNAMICMESH_API FMeshOcclusionMapBaker : public FMeshImageBaker
 {
@@ -41,4 +45,7 @@ public:
 protected:
 	TUniquePtr<TImageBuilder<FVector3f>> OcclusionBuilder;
 
-};
\ No newline at end of file
+};
+
+} // end namespace UE::Geometry
+} // end namespace UE
diff --git a/Engine/Plugins/Experimental/GeometryProcessing/Source/DynamicMesh/Public/Sampling/MeshSurfaceSampler.h b/Engine/Plugins/Experimental/GeometryProcessing/Source/DynamicMesh/Public/Sampling/MeshSurfaceSampler.h
--- a/Engine/Plugins/Experimental/GeometryProcessing/Source/DynamicMesh/Public/Sampling/MeshSurfaceSampler.h
+++ b/Engine/Plugins/Experimental/GeometryProcessing/Source/DynamicMesh/Public/Sampling/MeshSurfaceSampler.h
@@ -6,6 +6,10 @@
 #include "DynamicMeshAttributeSet.h"
 #include "DynamicMeshAABBTree3.h"
 
+namespace UE
+{
+namespace Geometry
+{
 
 /**
  * Wrapper around a Mesh and UV Overlay that provides UVs triangles as vertices.
@@ -295,3 +299,6 @@ bool TMeshSurfaceUVSampler<SampleType>::CachedSampleUV(const FMeshUVSampleInfo&
 
 
 
+} // end namespace UE::Geometry
+} // end namespace UE
+
diff --git a/Engine/Plugins/Experimental/GeometryProcessing/Source/DynamicMesh/Public/Selections/MeshEdgeSelection.h b/Engine/Plugins/Experimental/GeometryProcessing/Source/DynamicMesh/Public/Selections/MeshEdgeSelection.h
--- a/Engine/Plugins/Experimental/GeometryProcessing/Source/DynamicMesh/Public/Selections/MeshEdgeSelection.h
+++ b/Engine/Plugins/Experimental/GeometryProcessing/Source/DynamicMesh/Public/Selections/MeshEdgeSelection.h
@@ -9,6 +9,12 @@
 #include "Containers/Set.h"
 #include "Selections/MeshVertexSelection.h"
 
+namespace UE
+{
+namespace Geometry
+{
+
+
 /**
  * Currently a thin wrapper of a TSet<int> of Edge IDs paired with a Mesh; the backing storage will likely change as we need to optimize in the future
  */
@@ -170,4 +176,7 @@ public:
 		Selected.Empty();
 	}
 
-};
\ No newline at end of file
+};
+
+} // end namespace UE::Geometry
+} // end namespace UE
diff --git a/Engine/Plugins/Experimental/GeometryProcessing/Source/DynamicMesh/Public/ShapeApproximation/MeshSimpleShapeApproximation.h b/Engine/Plugins/Experimental/GeometryProcessing/Source/DynamicMesh/Public/ShapeApproximation/MeshSimpleShapeApproximation.h
--- a/Engine/Plugins/Experimental/GeometryProcessing/Source/DynamicMesh/Public/ShapeApproximation/MeshSimpleShapeApproximation.h
+++ b/Engine/Plugins/Experimental/GeometryProcessing/Source/DynamicMesh/Public/ShapeApproximation/MeshSimpleShapeApproximation.h
@@ -8,6 +8,9 @@
 
 class FProgressCancel;
 
+namespace UE {
+namespace Geometry {
+
 /**
  * EDetectedSimpleShapeType is used to identify auto-detected simple shapes for a mesh/etc
  */
@@ -166,3 +169,5 @@ protected:
 };
 
 
+} // end namespace UE::Geometry
+} // end namespace UE
diff --git a/Engine/Plugins/Experimental/GeometryProcessing/Source/DynamicMesh/Public/ShapeApproximation/SimpleShapeSet3.h b/Engine/Plugins/Experimental/GeometryProcessing/Source/DynamicMesh/Public/ShapeApproximation/SimpleShapeSet3.h
--- a/Engine/Plugins/Experimental/GeometryProcessing/Source/DynamicMesh/Public/ShapeApproximation/SimpleShapeSet3.h
+++ b/Engine/Plugins/Experimental/GeometryProcessing/Source/DynamicMesh/Public/ShapeApproximation/SimpleShapeSet3.h
@@ -12,6 +12,10 @@
 #include "TransformTypes.h"
 
 
+namespace UE {
+namespace Geometry {
+
+
 /**
  * Supported/known types of Simple Shapes
  */
@@ -152,3 +156,6 @@ struct DYNAMICMESH_API FSimpleShapeSet3d
 	void ApplyTransform(const FTransform3d& Transform);
 };
 
+
+} // end namespace UE::Geometry
+} // end namespace UE
\ No newline at end of file
diff --git a/Engine/Plugins/Experimental/GeometryProcessing/Source/DynamicMesh/Public/Solvers/ConstrainedMeshSolver.h b/Engine/Plugins/Experimental/GeometryProcessing/Source/DynamicMesh/Public/Solvers/ConstrainedMeshSolver.h
--- a/Engine/Plugins/Experimental/GeometryProcessing/Source/DynamicMesh/Public/Solvers/ConstrainedMeshSolver.h
+++ b/Engine/Plugins/Experimental/GeometryProcessing/Source/DynamicMesh/Public/Solvers/ConstrainedMeshSolver.h
@@ -5,12 +5,13 @@
 #include "CoreMinimal.h"
 #include "VectorTypes.h"
 
-class FDynamicMesh3;
+PREDECLARE_GEOMETRY(class FDynamicMesh3);
 
 namespace UE
 {
 	namespace Solvers
 	{
+		using namespace UE::Geometry;
 
 		/**
 		 * Interface to a index-based deformation solver for a 3D mesh vertex set that supports weighted point constraints.
diff --git a/Engine/Plugins/Experimental/GeometryProcessing/Source/DynamicMesh/Public/Solvers/MeshLinearization.h b/Engine/Plugins/Experimental/GeometryProcessing/Source/DynamicMesh/Public/Solvers/MeshLinearization.h
--- a/Engine/Plugins/Experimental/GeometryProcessing/Source/DynamicMesh/Public/Solvers/MeshLinearization.h
+++ b/Engine/Plugins/Experimental/GeometryProcessing/Source/DynamicMesh/Public/Solvers/MeshLinearization.h
@@ -7,6 +7,13 @@
 #include "DynamicMesh3.h"
 
 
+namespace UE
+{
+namespace Geometry
+{
+
+class FDynamicMesh3;
+
 /**
 * Used linearize the VtxIds in a mesh as a single array and allow mapping from array offset to mesh VtxId.
 * Generally, the array offset will correspond to a matrix row when forming a Laplacian.
@@ -123,3 +130,7 @@ public:
 private:
 	FTriangleLinearization(const FTriangleLinearization&);
 };
+
+
+} // end namespace UE::Geometry
+} // end namespace UE
diff --git a/Engine/Plugins/Experimental/GeometryProcessing/Source/DynamicMesh/Public/Solvers/MeshSmoothing.h b/Engine/Plugins/Experimental/GeometryProcessing/Source/DynamicMesh/Public/Solvers/MeshSmoothing.h
--- a/Engine/Plugins/Experimental/GeometryProcessing/Source/DynamicMesh/Public/Solvers/MeshSmoothing.h
+++ b/Engine/Plugins/Experimental/GeometryProcessing/Source/DynamicMesh/Public/Solvers/MeshSmoothing.h
@@ -10,6 +10,9 @@ namespace UE
 {
 	namespace MeshDeformation
 	{
+		using namespace UE::Geometry;
+		using UE::Geometry::FVector3d;
+
 		/**
 		*
 		* Note: for discussion of implicit / explicit integration of diffusion and biharmonic equations
diff --git a/Engine/Plugins/Experimental/GeometryProcessing/Source/GeometricObjects/Private/Util/ColorConstants.cpp b/Engine/Plugins/Experimental/GeometryProcessing/Source/GeometricObjects/Private/Util/ColorConstants.cpp
--- a/Engine/Plugins/Experimental/GeometryProcessing/Source/GeometricObjects/Private/Util/ColorConstants.cpp
+++ b/Engine/Plugins/Experimental/GeometryProcessing/Source/GeometricObjects/Private/Util/ColorConstants.cpp
@@ -2,8 +2,7 @@
 
 #include "Util/ColorConstants.h"
 
-
-FColor LinearColors::SelectFColor(int32 Index)
+FColor UE::Geometry::LinearColors::SelectFColor(int32 Index)
 {
 	static const FColor ColorMap[] = {
 		SpringGreen3b(), Plum3b(), Khaki3b(),
diff --git a/Engine/Plugins/Experimental/GeometryProcessing/Source/GeometricObjects/Public/CapsuleTypes.h b/Engine/Plugins/Experimental/GeometryProcessing/Source/GeometricObjects/Public/CapsuleTypes.h
--- a/Engine/Plugins/Experimental/GeometryProcessing/Source/GeometricObjects/Public/CapsuleTypes.h
+++ b/Engine/Plugins/Experimental/GeometryProcessing/Source/GeometricObjects/Public/CapsuleTypes.h
@@ -9,6 +9,11 @@
 #include "OrientedBoxTypes.h"
 
 
+namespace UE 
+{
+namespace Geometry 
+{
+
 /*
  * 3D Capsule stored as Line Segment and Radius
  */
@@ -118,4 +123,8 @@ public:
 
 
 typedef TCapsule3<float> FCapsule3f;
-typedef TCapsule3<double> FCapsule3d;
\ No newline at end of file
+typedef TCapsule3<double> FCapsule3d;
+
+
+} // end namespace UE::Geometry
+} // end namespace UE
diff --git a/Engine/Plugins/Experimental/GeometryProcessing/Source/GeometricObjects/Public/CircleTypes.h b/Engine/Plugins/Experimental/GeometryProcessing/Source/GeometricObjects/Public/CircleTypes.h
--- a/Engine/Plugins/Experimental/GeometryProcessing/Source/GeometricObjects/Public/CircleTypes.h
+++ b/Engine/Plugins/Experimental/GeometryProcessing/Source/GeometricObjects/Public/CircleTypes.h
@@ -1,7 +1,5 @@
 // Copyright Epic Games, Inc. All Rights Reserved.
 
-// Port of WildMagic DistPoint3TCircle3
-
 #pragma once
 
 #include "VectorTypes.h"
@@ -10,6 +8,11 @@
 #include "OrientedBoxTypes.h"
 #include "FrameTypes.h"
 
+namespace UE 
+{
+namespace Geometry 
+{
+
 template<typename RealType>
 struct TCircle2
 {
@@ -271,3 +274,6 @@ struct TCircle3
 typedef TCircle3<float> FCircle3f;
 typedef TCircle3<double> FCircle3d;
 
+
+} // end namespace UE::Geometry
+} // end namespace UE
diff --git a/Engine/Plugins/Experimental/GeometryProcessing/Source/GeometricObjects/Public/Curve/DynamicGraph.h b/Engine/Plugins/Experimental/GeometryProcessing/Source/GeometricObjects/Public/Curve/DynamicGraph.h
--- a/Engine/Plugins/Experimental/GeometryProcessing/Source/GeometricObjects/Public/Curve/DynamicGraph.h
+++ b/Engine/Plugins/Experimental/GeometryProcessing/Source/GeometricObjects/Public/Curve/DynamicGraph.h
@@ -15,6 +15,12 @@
 #include "Util/SmallListSet.h"
 #include "VectorUtil.h"
 
+namespace UE
+{
+namespace Geometry
+{
+
+
 class FDynamicGraph
 {
 public:
@@ -699,3 +705,7 @@ protected:
 		return AppendVertex();
 	}
 };
+
+
+} // end namespace UE::Geometry
+} // end namespace UE
\ No newline at end of file
diff --git a/Engine/Plugins/Experimental/GeometryProcessing/Source/GeometricObjects/Public/Curve/PlanarComplex.h b/Engine/Plugins/Experimental/GeometryProcessing/Source/GeometricObjects/Public/Curve/PlanarComplex.h
--- a/Engine/Plugins/Experimental/GeometryProcessing/Source/GeometricObjects/Public/Curve/PlanarComplex.h
+++ b/Engine/Plugins/Experimental/GeometryProcessing/Source/GeometricObjects/Public/Curve/PlanarComplex.h
@@ -5,14 +5,17 @@
 // Very approximately ported from geometry3sharp's PlanarComplex
 // Convert a set of polygons into general polygon form, tracking the outer-polygon and hole-polygon indices
 
-#pragma once
-
 #include "SegmentTypes.h"
 #include "VectorTypes.h"
 #include "BoxTypes.h"
 #include "Polygon2.h"
 #include "Curve/GeneralPolygon2.h"
 
+namespace UE
+{
+namespace Geometry
+{
+
 template <typename RealType>
 struct TPlanarComplex
 {
@@ -175,3 +178,7 @@ struct TPlanarComplex
 
 typedef TPlanarComplex<double> FPlanarComplexd;
 typedef TPlanarComplex<float> FPlanarComplexf;
+
+
+} // end namespace UE::Geometry
+} // end namespace UE
diff --git a/Engine/Plugins/Experimental/GeometryProcessing/Source/GeometricObjects/Public/Distance/DistLine3Line3.h b/Engine/Plugins/Experimental/GeometryProcessing/Source/GeometricObjects/Public/Distance/DistLine3Line3.h
--- a/Engine/Plugins/Experimental/GeometryProcessing/Source/GeometricObjects/Public/Distance/DistLine3Line3.h
+++ b/Engine/Plugins/Experimental/GeometryProcessing/Source/GeometricObjects/Public/Distance/DistLine3Line3.h
@@ -7,6 +7,11 @@
 #include "VectorTypes.h"
 #include "LineTypes.h"
 
+namespace UE
+{
+namespace Geometry
+{
+
 /**
  * Compute distance between two 3D lines
  */
@@ -94,4 +99,7 @@ public:
 };
 
 typedef TDistLine3Line3<float> FDistLine3Line3f;
-typedef TDistLine3Line3<double> FDistLine3Line3d;
\ No newline at end of file
+typedef TDistLine3Line3<double> FDistLine3Line3d;
+
+} // end namespace UE::Geometry
+} // end namespace UE
diff --git a/Engine/Plugins/Experimental/GeometryProcessing/Source/GeometricObjects/Public/Distance/DistLine3Ray3.h b/Engine/Plugins/Experimental/GeometryProcessing/Source/GeometricObjects/Public/Distance/DistLine3Ray3.h
--- a/Engine/Plugins/Experimental/GeometryProcessing/Source/GeometricObjects/Public/Distance/DistLine3Ray3.h
+++ b/Engine/Plugins/Experimental/GeometryProcessing/Source/GeometricObjects/Public/Distance/DistLine3Ray3.h
@@ -8,6 +8,11 @@
 #include "LineTypes.h"
 #include "RayTypes.h"
 
+namespace UE
+{
+namespace Geometry
+{
+
 /**
  * Compute distance between 3D line and 3D ray
  */
@@ -104,4 +109,7 @@ public:
 };
 
 typedef TDistLine3Ray3<float> FDistLine3Ray3f;
-typedef TDistLine3Ray3<double> FDistLine3Ray3d;
\ No newline at end of file
+typedef TDistLine3Ray3<double> FDistLine3Ray3d;
+
+} // end namespace UE::Geometry
+} // end namespace UE
diff --git a/Engine/Plugins/Experimental/GeometryProcessing/Source/GeometricObjects/Public/Distance/DistLine3Segment3.h b/Engine/Plugins/Experimental/GeometryProcessing/Source/GeometricObjects/Public/Distance/DistLine3Segment3.h
--- a/Engine/Plugins/Experimental/GeometryProcessing/Source/GeometricObjects/Public/Distance/DistLine3Segment3.h
+++ b/Engine/Plugins/Experimental/GeometryProcessing/Source/GeometricObjects/Public/Distance/DistLine3Segment3.h
@@ -9,6 +9,10 @@
 #include "SegmentTypes.h"
 #include "LineTypes.h"
 
+namespace UE
+{
+namespace Geometry
+{
 
 /**
 * Compute unsigned distance between 3D line and 3D segment
@@ -120,3 +124,5 @@ public:
 typedef TDistLine3Segment3<float> FDistLine3Segment3f;
 typedef TDistLine3Segment3<double> FDistLine3Segment3d;
 
+} // end namespace UE::Geometry
+} // end namespace UE
diff --git a/Engine/Plugins/Experimental/GeometryProcessing/Source/GeometricObjects/Public/Distance/DistRay3Segment3.h b/Engine/Plugins/Experimental/GeometryProcessing/Source/GeometricObjects/Public/Distance/DistRay3Segment3.h
--- a/Engine/Plugins/Experimental/GeometryProcessing/Source/GeometricObjects/Public/Distance/DistRay3Segment3.h
+++ b/Engine/Plugins/Experimental/GeometryProcessing/Source/GeometricObjects/Public/Distance/DistRay3Segment3.h
@@ -8,6 +8,11 @@
 #include "SegmentTypes.h"
 #include "RayTypes.h"
 
+namespace UE
+{
+namespace Geometry
+{
+
 /**
  * Compute distance between 3D ray and 3D segment
  */
@@ -386,4 +391,7 @@ public:
 };
 
 typedef TDistRay3Segment3<float> FDistRay3Segment3f;
-typedef TDistRay3Segment3<double> FDistRay3Segment3d;
\ No newline at end of file
+typedef TDistRay3Segment3<double> FDistRay3Segment3d;
+
+} // end namespace UE::Geometry
+} // end namespace UE
