GeometryProcessing: move FVector2/3/4/f/d back out of UE::Geometry:: namespace
#rnx

diff --git a/Engine/Plugins/Experimental/ChaosEditor/Source/FractureEditor/Private/FractureToolUV.cpp b/Engine/Plugins/Experimental/ChaosEditor/Source/FractureEditor/Private/FractureToolUV.cpp
--- a/Engine/Plugins/Experimental/ChaosEditor/Source/FractureEditor/Private/FractureToolUV.cpp
+++ b/Engine/Plugins/Experimental/ChaosEditor/Source/FractureEditor/Private/FractureToolUV.cpp
@@ -218,9 +218,9 @@ int32 UFractureToolAutoUV::ExecuteFracture(const FFractureToolContext& FractureC
 		FImageDimensions Dimensions(OutputRes, OutputRes);
 		TextureBuilder.Initialize(FTexture2DBuilder::ETextureType::Color, Dimensions);
 
-		UE::Geometry::TImageBuilder<UE::Geometry::FVector3f> ImageBuilder;
+		UE::Geometry::TImageBuilder<FVector3f> ImageBuilder;
 		ImageBuilder.SetDimensions(Dimensions);
-		ImageBuilder.Clear(UE::Geometry::FVector3f(0, 0, 0));
+		ImageBuilder.Clear(FVector3f(0, 0, 0));
 
 		UE::PlanarCut::TextureInternalSurfaces(Collection, AutoUVSettings->MaxDistance, FMath::CeilToInt(AutoUVSettings->GutterSize), ImageBuilder);
 
diff --git a/Engine/Plugins/Experimental/Enterprise/DataprepGeometryOperations/Source/Private/DataprepGeometryOperations.cpp b/Engine/Plugins/Experimental/Enterprise/DataprepGeometryOperations/Source/Private/DataprepGeometryOperations.cpp
--- a/Engine/Plugins/Experimental/Enterprise/DataprepGeometryOperations/Source/Private/DataprepGeometryOperations.cpp
+++ b/Engine/Plugins/Experimental/Enterprise/DataprepGeometryOperations/Source/Private/DataprepGeometryOperations.cpp
@@ -296,12 +296,12 @@ void UDataprepBakeTransformOperation::OnExecution_Implementation(const FDataprep
 				ToBakePart.SetRotation(ComponentToWorld.GetRotation());
 				NewWorldPart.SetRotation(UE::Geometry::FQuaterniond::Identity());
 			}
-			UE::Geometry::FVector3d ScaleVec = ComponentToWorld.GetScale();
+			FVector3d ScaleVec = ComponentToWorld.GetScale();
 
-			UE::Geometry::FVector3d AbsScales(FMathd::Abs(ScaleVec.X), FMathd::Abs(ScaleVec.Y), FMathd::Abs(ScaleVec.Z));
+			FVector3d AbsScales(FMathd::Abs(ScaleVec.X), FMathd::Abs(ScaleVec.Y), FMathd::Abs(ScaleVec.Z));
 			double RemainingUniformScale = AbsScales.X;
 			{
-				UE::Geometry::FVector3d Dists;
+				FVector3d Dists;
 				for (int SubIdx = 0; SubIdx < 3; SubIdx++)
 				{
 					int OtherA = (SubIdx + 1) % 3;
@@ -326,12 +326,12 @@ void UDataprepBakeTransformOperation::OnExecution_Implementation(const FDataprep
 			{
 			case EBakeScaleMethod::BakeFullScale:
 				ToBakePart.SetScale(ScaleVec);
-				NewWorldPart.SetScale(UE::Geometry::FVector3d::One());
+				NewWorldPart.SetScale(FVector3d::One());
 				break;
 			case EBakeScaleMethod::BakeNonuniformScale:
 				check(RemainingUniformScale > FLT_MIN); // avoid baking a ~zero scale
 				ToBakePart.SetScale(ScaleVec / RemainingUniformScale);
-				NewWorldPart.SetScale(UE::Geometry::FVector3d(RemainingUniformScale, RemainingUniformScale, RemainingUniformScale));
+				NewWorldPart.SetScale(FVector3d(RemainingUniformScale, RemainingUniformScale, RemainingUniformScale));
 				break;
 			case EBakeScaleMethod::DoNotBakeScale:
 				break;
@@ -351,7 +351,7 @@ void UDataprepBakeTransformOperation::OnExecution_Implementation(const FDataprep
 			if (bRecenterPivot)
 			{
 				FBox BBox = MeshDescription->ComputeBoundingBox();
-				UE::Geometry::FVector3d Center(BBox.GetCenter());
+				FVector3d Center(BBox.GetCenter());
 				UE::Geometry::FFrame3d LocalFrame(Center);
 				ToBakePart.SetTranslation(ToBakePart.GetTranslation() - Center);
 				NewWorldPart.SetTranslation(NewWorldPart.GetTranslation() + NewWorldPart.TransformVector(Center));
diff --git a/Engine/Plugins/Experimental/GeometryFlow/Source/GeometryFlowCore/Public/GeometryFlowCoreNodes.h b/Engine/Plugins/Experimental/GeometryFlow/Source/GeometryFlowCore/Public/GeometryFlowCoreNodes.h
--- a/Engine/Plugins/Experimental/GeometryFlow/Source/GeometryFlowCore/Public/GeometryFlowCoreNodes.h
+++ b/Engine/Plugins/Experimental/GeometryFlow/Source/GeometryFlowCore/Public/GeometryFlowCoreNodes.h
@@ -12,8 +12,6 @@ namespace GeometryFlow
 {
 
 using namespace UE::Geometry;
-using UE::Geometry::FVector3d;
-using UE::Geometry::FVector3f;
 
 
 template<typename T, int32 StorageTypeIdentifier>
diff --git a/Engine/Plugins/Experimental/GeometryProcessing/Source/DynamicMesh/Public/InfoTypes.h b/Engine/Plugins/Experimental/GeometryProcessing/Source/DynamicMesh/Public/InfoTypes.h
--- a/Engine/Plugins/Experimental/GeometryProcessing/Source/DynamicMesh/Public/InfoTypes.h
+++ b/Engine/Plugins/Experimental/GeometryProcessing/Source/DynamicMesh/Public/InfoTypes.h
@@ -41,7 +41,6 @@ struct FVertexInfo
 namespace DynamicMeshInfo
 {
 using namespace UE::Geometry;
-using FVector3d = UE::Geometry::FVector3d;
 
 /** Information about the mesh elements created by a call to SplitEdge() */
 struct FEdgeSplitInfo
diff --git a/Engine/Plugins/Experimental/GeometryProcessing/Source/DynamicMesh/Public/MeshTransforms.h b/Engine/Plugins/Experimental/GeometryProcessing/Source/DynamicMesh/Public/MeshTransforms.h
--- a/Engine/Plugins/Experimental/GeometryProcessing/Source/DynamicMesh/Public/MeshTransforms.h
+++ b/Engine/Plugins/Experimental/GeometryProcessing/Source/DynamicMesh/Public/MeshTransforms.h
@@ -13,8 +13,6 @@
 namespace MeshTransforms
 {
 	using namespace UE::Geometry;
-	using FVector3d = UE::Geometry::FVector3d;
-	using FVector3f = UE::Geometry::FVector3f;
 
 	/**
 	 * Apply Translation to vertex positions of Mesh. Does not modify any other attributes.
diff --git a/Engine/Plugins/Experimental/GeometryProcessing/Source/DynamicMesh/Public/Solvers/MeshSmoothing.h b/Engine/Plugins/Experimental/GeometryProcessing/Source/DynamicMesh/Public/Solvers/MeshSmoothing.h
--- a/Engine/Plugins/Experimental/GeometryProcessing/Source/DynamicMesh/Public/Solvers/MeshSmoothing.h
+++ b/Engine/Plugins/Experimental/GeometryProcessing/Source/DynamicMesh/Public/Solvers/MeshSmoothing.h
@@ -11,7 +11,6 @@ namespace UE
 	namespace MeshDeformation
 	{
 		using namespace UE::Geometry;
-		using UE::Geometry::FVector3d;
 
 		/**
 		*
diff --git a/Engine/Plugins/Experimental/GeometryProcessing/Source/GeometricObjects/Public/ExplicitUseGeometryMathTypes.h b/Engine/Plugins/Experimental/GeometryProcessing/Source/GeometricObjects/Public/ExplicitUseGeometryMathTypes.h
--- a/Engine/Plugins/Experimental/GeometryProcessing/Source/GeometricObjects/Public/ExplicitUseGeometryMathTypes.h
+++ b/Engine/Plugins/Experimental/GeometryProcessing/Source/GeometricObjects/Public/ExplicitUseGeometryMathTypes.h
@@ -49,14 +49,14 @@ using UE::Geometry::FIndex2i;
 using UE::Geometry::FIndex3i;
 using UE::Geometry::FIndex4i;
 
-using UE::Geometry::FVector2f;
-using UE::Geometry::FVector2d;
+//using UE::Geometry::FVector2f;
+//using UE::Geometry::FVector2d;
 using UE::Geometry::FVector2i;
-using UE::Geometry::FVector3f;
-using UE::Geometry::FVector3d;
+//using UE::Geometry::FVector3f;
+//using UE::Geometry::FVector3d;
 using UE::Geometry::FVector3i;
-using UE::Geometry::FVector4f;
-using UE::Geometry::FVector4d;
+//using UE::Geometry::FVector4f;
+//using UE::Geometry::FVector4d;
 using UE::Geometry::FVector4i;
 
 using UE::Geometry::FMatrix3f;
diff --git a/Engine/Plugins/Experimental/GeometryProcessing/Source/GeometricObjects/Public/MeshAdapterTransforms.h b/Engine/Plugins/Experimental/GeometryProcessing/Source/GeometricObjects/Public/MeshAdapterTransforms.h
--- a/Engine/Plugins/Experimental/GeometryProcessing/Source/GeometricObjects/Public/MeshAdapterTransforms.h
+++ b/Engine/Plugins/Experimental/GeometryProcessing/Source/GeometricObjects/Public/MeshAdapterTransforms.h
@@ -22,8 +22,6 @@
 namespace MeshAdapterTransforms
 {
 	using namespace UE::Geometry;
-	using FVector3d = UE::Geometry::FVector3d;
-	using FVector3f = UE::Geometry::FVector3f;
 	using FTransform3d = UE::Geometry::FTransform3d;
 	using FFrame3d = UE::Geometry::FFrame3d;
 
diff --git a/Engine/Plugins/Experimental/GeometryProcessing/Source/GeometricObjects/Public/Spatial/FastWinding.h b/Engine/Plugins/Experimental/GeometryProcessing/Source/GeometricObjects/Public/Spatial/FastWinding.h
--- a/Engine/Plugins/Experimental/GeometryProcessing/Source/GeometricObjects/Public/Spatial/FastWinding.h
+++ b/Engine/Plugins/Experimental/GeometryProcessing/Source/GeometricObjects/Public/Spatial/FastWinding.h
@@ -15,7 +15,6 @@ namespace FastTriWinding
 {
 	using namespace UE::Geometry;
 	using FTriangle3d = UE::Geometry::FTriangle3d;
-	using FVector3d = UE::Geometry::FVector3d;
 	using FMatrix3d = UE::Geometry::FMatrix3d;
 
 	/**
diff --git a/Engine/Plugins/Experimental/GeometryProcessing/Source/GeometricObjects/Public/Spatial/MeshAABBTree3.h b/Engine/Plugins/Experimental/GeometryProcessing/Source/GeometricObjects/Public/Spatial/MeshAABBTree3.h
--- a/Engine/Plugins/Experimental/GeometryProcessing/Source/GeometricObjects/Public/Spatial/MeshAABBTree3.h
+++ b/Engine/Plugins/Experimental/GeometryProcessing/Source/GeometricObjects/Public/Spatial/MeshAABBTree3.h
@@ -16,7 +16,6 @@
 namespace MeshIntersection
 {
 	using namespace UE::Geometry;
-	using FVector3d = UE::Geometry::FVector3d;
 
 	/**
 	 * Intersection query result types for triangle mesh intersections
diff --git a/Engine/Plugins/Experimental/GeometryProcessing/Source/GeometricObjects/Public/Util/RevolveUtil.h b/Engine/Plugins/Experimental/GeometryProcessing/Source/GeometricObjects/Public/Util/RevolveUtil.h
--- a/Engine/Plugins/Experimental/GeometryProcessing/Source/GeometricObjects/Public/Util/RevolveUtil.h
+++ b/Engine/Plugins/Experimental/GeometryProcessing/Source/GeometricObjects/Public/Util/RevolveUtil.h
@@ -12,7 +12,6 @@ PREDECLARE_GEOMETRY(class FCurveSweepOp);
 namespace RevolveUtil
 {
 	using namespace UE::Geometry;
-	using UE::Geometry::FVector3d;
 	using UE::Geometry::FFrame3d;
 
 	/** 
diff --git a/Engine/Plugins/Experimental/GeometryProcessing/Source/GeometricObjects/Public/VectorTypes.h b/Engine/Plugins/Experimental/GeometryProcessing/Source/GeometricObjects/Public/VectorTypes.h
--- a/Engine/Plugins/Experimental/GeometryProcessing/Source/GeometricObjects/Public/VectorTypes.h
+++ b/Engine/Plugins/Experimental/GeometryProcessing/Source/GeometricObjects/Public/VectorTypes.h
@@ -781,8 +781,6 @@ std::ostream& operator<<(std::ostream& os, const FVector3<RealType>& Vec)
 	return os;
 }
 
-typedef FVector3<float> FVector3f;
-typedef FVector3<double> FVector3d;
 typedef FVector3<int> FVector3i;
 
 template <typename T>
@@ -825,8 +823,6 @@ std::ostream& operator<<(std::ostream& os, const FVector2<RealType>& Vec)
 	return os;
 }
 
-typedef FVector2<float> FVector2f;
-typedef FVector2<double> FVector2d;
 typedef FVector2<int> FVector2i;
 
 
@@ -1116,8 +1112,6 @@ std::ostream& operator<<(std::ostream& os, const TVector4<RealType>& Vec)
 	return os;
 }
 
-typedef TVector4<float> FVector4f;
-typedef TVector4<double> FVector4d;
 typedef TVector4<int> FVector4i;
 
 template <typename T>
@@ -1130,4 +1124,14 @@ FORCEINLINE uint32 GetTypeHash(const TVector4<T>& Vector)
 
 
 } // end namespace UE::Geometry
-} // end namespace UE
\ No newline at end of file
+} // end namespace UE
+
+
+typedef UE::Geometry::FVector2<float> FVector2f;
+typedef UE::Geometry::FVector2<double> FVector2d;
+
+typedef UE::Geometry::FVector3<float> FVector3f;
+typedef UE::Geometry::FVector3<double> FVector3d;
+
+typedef UE::Geometry::TVector4<float> FVector4f;
+typedef UE::Geometry::TVector4<double> FVector4d;
\ No newline at end of file
diff --git a/Engine/Plugins/Experimental/GeometryProcessing/Source/GeometricObjects/Public/VectorUtil.h b/Engine/Plugins/Experimental/GeometryProcessing/Source/GeometricObjects/Public/VectorUtil.h
--- a/Engine/Plugins/Experimental/GeometryProcessing/Source/GeometricObjects/Public/VectorUtil.h
+++ b/Engine/Plugins/Experimental/GeometryProcessing/Source/GeometricObjects/Public/VectorUtil.h
@@ -29,8 +29,6 @@ enum class EIntersectionType
 namespace VectorUtil
 {
 	using namespace UE::Geometry;
-	using FVector2d = UE::Geometry::FVector2d;
-	using FVector3d = UE::Geometry::FVector3d;
 
 	/**
 	 * @return true if all components of V are finite
diff --git a/Engine/Plugins/Experimental/GeometryProcessing/Source/MeshConversion/Public/MeshDescriptionAdapter.h b/Engine/Plugins/Experimental/GeometryProcessing/Source/MeshConversion/Public/MeshDescriptionAdapter.h
--- a/Engine/Plugins/Experimental/GeometryProcessing/Source/MeshConversion/Public/MeshDescriptionAdapter.h
+++ b/Engine/Plugins/Experimental/GeometryProcessing/Source/MeshConversion/Public/MeshDescriptionAdapter.h
@@ -19,8 +19,6 @@
  */
 struct /*MESHCONVERSION_API*/ FMeshDescriptionTriangleMeshAdapter
 {
-	using FVector3d = UE::Geometry::FVector3d;
-	using FVector3f = UE::Geometry::FVector3f;
 	using FIndex3i = UE::Geometry::FIndex3i;
 protected:
 	const FMeshDescription* Mesh;
@@ -115,8 +113,6 @@ public:
  */
 struct /*MESHCONVERSION_API*/ FMeshDescriptionEditableTriangleMeshAdapter
 {
-	using FVector3d = UE::Geometry::FVector3d;
-	using FVector3f = UE::Geometry::FVector3f;
 	using FIndex3i = UE::Geometry::FIndex3i;
 protected:
 	FMeshDescription* Mesh;
diff --git a/Engine/Plugins/Experimental/MeshLODToolset/Source/MeshLODToolset/Private/Graphs/GenerateMeshLODGraph.cpp b/Engine/Plugins/Experimental/MeshLODToolset/Source/MeshLODToolset/Private/Graphs/GenerateMeshLODGraph.cpp
--- a/Engine/Plugins/Experimental/MeshLODToolset/Source/MeshLODToolset/Private/Graphs/GenerateMeshLODGraph.cpp
+++ b/Engine/Plugins/Experimental/MeshLODToolset/Source/MeshLODToolset/Private/Graphs/GenerateMeshLODGraph.cpp
@@ -636,7 +636,7 @@ void FGenerateMeshLODGraph::BuildGraph()
 
 
 
-int32 FGenerateMeshLODGraph::AppendTextureBakeNode(const TImageBuilder<UE::Geometry::FVector4f>& SourceImage, const FString& Identifier)
+int32 FGenerateMeshLODGraph::AppendTextureBakeNode(const TImageBuilder<FVector4f>& SourceImage, const FString& Identifier)
 {
 	FBakeTextureGraphInfo NewNode;
 	NewNode.Index = BakeTextureNodes.Num();
diff --git a/Engine/Plugins/Experimental/MeshLODToolset/Source/MeshLODToolset/Public/Graphs/GenerateMeshLODGraph.h b/Engine/Plugins/Experimental/MeshLODToolset/Source/MeshLODToolset/Public/Graphs/GenerateMeshLODGraph.h
--- a/Engine/Plugins/Experimental/MeshLODToolset/Source/MeshLODToolset/Public/Graphs/GenerateMeshLODGraph.h
+++ b/Engine/Plugins/Experimental/MeshLODToolset/Source/MeshLODToolset/Public/Graphs/GenerateMeshLODGraph.h
@@ -28,7 +28,6 @@ struct FMeshLODGraphPreFilterSettings
 
 class FGenerateMeshLODGraph
 {
-	using FVector4f = UE::Geometry::FVector4f;
 public:
 	void BuildGraph();
 
diff --git a/Engine/Plugins/Experimental/MeshLODToolset/Source/MeshLODToolset/Public/Graphs/GenerateStaticMeshLODProcess.h b/Engine/Plugins/Experimental/MeshLODToolset/Source/MeshLODToolset/Public/Graphs/GenerateStaticMeshLODProcess.h
--- a/Engine/Plugins/Experimental/MeshLODToolset/Source/MeshLODToolset/Public/Graphs/GenerateStaticMeshLODProcess.h
+++ b/Engine/Plugins/Experimental/MeshLODToolset/Source/MeshLODToolset/Public/Graphs/GenerateStaticMeshLODProcess.h
@@ -161,7 +161,6 @@ UCLASS(Transient)
 class UGenerateStaticMeshLODProcess : public UObject
 {
 	GENERATED_BODY()
-	using FVector4f = UE::Geometry::FVector4f;
 public:
 
 	bool Initialize(UStaticMesh* SourceMesh);
diff --git a/Engine/Plugins/Experimental/MeshModelingToolset/Source/MeshModelingTools/Private/Properties/MeshAnalysisProperties.cpp b/Engine/Plugins/Experimental/MeshModelingToolset/Source/MeshModelingTools/Private/Properties/MeshAnalysisProperties.cpp
--- a/Engine/Plugins/Experimental/MeshModelingToolset/Source/MeshModelingTools/Private/Properties/MeshAnalysisProperties.cpp
+++ b/Engine/Plugins/Experimental/MeshModelingToolset/Source/MeshModelingTools/Private/Properties/MeshAnalysisProperties.cpp
@@ -14,7 +14,7 @@ using namespace UE::Geometry;
 void UMeshAnalysisProperties::Update(const FDynamicMesh3& MeshIn, const FTransform& Transform)
 {
 	FTriangleMeshAdapterd TransformedMesh = MeshAdapterUtil::MakeTransformedDynamicMeshAdapter(&MeshIn, Transform);
-	UE::Geometry::FVector2d VolArea = TMeshQueries<FTriangleMeshAdapterd>::GetVolumeArea(TransformedMesh);
+	FVector2d VolArea = TMeshQueries<FTriangleMeshAdapterd>::GetVolumeArea(TransformedMesh);
 	this->SurfaceArea = FString::Printf(TEXT("%.2f m^2"), VolArea.Y / 10000);
 	this->Volume = FString::Printf(TEXT("%.2f m^3"), VolArea.X / 1000000);
 }
diff --git a/Engine/Plugins/Experimental/MeshModelingToolset/Source/MeshModelingTools/Public/AlignObjectsTool.h b/Engine/Plugins/Experimental/MeshModelingToolset/Source/MeshModelingTools/Public/AlignObjectsTool.h
--- a/Engine/Plugins/Experimental/MeshModelingToolset/Source/MeshModelingTools/Public/AlignObjectsTool.h
+++ b/Engine/Plugins/Experimental/MeshModelingToolset/Source/MeshModelingTools/Public/AlignObjectsTool.h
@@ -102,7 +102,6 @@ UCLASS()
 class MESHMODELINGTOOLS_API UAlignObjectsTool : public UMultiSelectionTool, public IClickDragBehaviorTarget
 {
 	GENERATED_BODY()
-	using FVector3d = UE::Geometry::FVector3d;
 	using FTransform3d = UE::Geometry::FTransform3d;
 	using FAxisAlignedBox3d = UE::Geometry::FAxisAlignedBox3d;
 public:
diff --git a/Engine/Plugins/Experimental/MeshModelingToolset/Source/MeshModelingTools/Public/DeformMeshPolygonsTool.h b/Engine/Plugins/Experimental/MeshModelingToolset/Source/MeshModelingTools/Public/DeformMeshPolygonsTool.h
--- a/Engine/Plugins/Experimental/MeshModelingToolset/Source/MeshModelingTools/Public/DeformMeshPolygonsTool.h
+++ b/Engine/Plugins/Experimental/MeshModelingToolset/Source/MeshModelingTools/Public/DeformMeshPolygonsTool.h
@@ -88,7 +88,6 @@ UCLASS()
 class MESHMODELINGTOOLS_API UDeformMeshPolygonsTransformProperties : public UInteractiveToolPropertySet
 {
 	GENERATED_BODY()
-	using FVector3d = UE::Geometry::FVector3d;
 
 public:
 	UDeformMeshPolygonsTransformProperties();
@@ -149,7 +148,6 @@ UCLASS()
 class MESHMODELINGTOOLS_API UDeformMeshPolygonsTool : public UMeshSurfacePointTool
 {
 	GENERATED_BODY()
-	using FVector3d = UE::Geometry::FVector3d;
 
 public:
 	UDeformMeshPolygonsTool();
diff --git a/Engine/Plugins/Experimental/MeshModelingToolset/Source/MeshModelingTools/Public/DisplaceMeshTool.h b/Engine/Plugins/Experimental/MeshModelingToolset/Source/MeshModelingTools/Public/DisplaceMeshTool.h
--- a/Engine/Plugins/Experimental/MeshModelingToolset/Source/MeshModelingTools/Public/DisplaceMeshTool.h
+++ b/Engine/Plugins/Experimental/MeshModelingToolset/Source/MeshModelingTools/Public/DisplaceMeshTool.h
@@ -240,7 +240,6 @@ UCLASS()
 class MESHMODELINGTOOLS_API UDisplaceMeshTool : public USingleSelectionTool
 {
 	GENERATED_BODY()
-	using FVector3d = UE::Geometry::FVector3d;
 public:
 	void Setup() final;
 	void Shutdown(EToolShutdownType ShutdownType) final;
diff --git a/Engine/Plugins/Experimental/MeshModelingToolset/Source/MeshModelingTools/Public/DrawAndRevolveTool.h b/Engine/Plugins/Experimental/MeshModelingToolset/Source/MeshModelingTools/Public/DrawAndRevolveTool.h
--- a/Engine/Plugins/Experimental/MeshModelingToolset/Source/MeshModelingTools/Public/DrawAndRevolveTool.h
+++ b/Engine/Plugins/Experimental/MeshModelingToolset/Source/MeshModelingTools/Public/DrawAndRevolveTool.h
@@ -116,8 +116,8 @@ protected:
 
 	// This information is replicated in the user-editable transform in the settings and in the PlaneMechanic
 	// plane, but the tool turned out to be much easier to write and edit with this decoupling.
-	UE::Geometry::FVector3d RevolutionAxisOrigin;
-	UE::Geometry::FVector3d RevolutionAxisDirection;
+	FVector3d RevolutionAxisOrigin;
+	FVector3d RevolutionAxisDirection;
 
 	bool bProfileCurveComplete = false;
 
diff --git a/Engine/Plugins/Experimental/MeshModelingToolset/Source/MeshModelingTools/Public/DrawPolyPathTool.h b/Engine/Plugins/Experimental/MeshModelingToolset/Source/MeshModelingTools/Public/DrawPolyPathTool.h
--- a/Engine/Plugins/Experimental/MeshModelingToolset/Source/MeshModelingTools/Public/DrawPolyPathTool.h
+++ b/Engine/Plugins/Experimental/MeshModelingToolset/Source/MeshModelingTools/Public/DrawPolyPathTool.h
@@ -191,7 +191,7 @@ protected:
 	TArray<UE::Geometry::FFrame3d> CurPathPoints;
 	TArray<double> OffsetScaleFactors;
 	TArray<double> ArcLengths;
-	TArray<UE::Geometry::FVector3d> CurPolyLine;
+	TArray<FVector3d> CurPolyLine;
 	double CurPathLength;
 	double CurOffsetDistance;
 	double CurHeight;
diff --git a/Engine/Plugins/Experimental/MeshModelingToolset/Source/MeshModelingTools/Public/DrawPolygonTool.h b/Engine/Plugins/Experimental/MeshModelingToolset/Source/MeshModelingTools/Public/DrawPolygonTool.h
--- a/Engine/Plugins/Experimental/MeshModelingToolset/Source/MeshModelingTools/Public/DrawPolygonTool.h
+++ b/Engine/Plugins/Experimental/MeshModelingToolset/Source/MeshModelingTools/Public/DrawPolygonTool.h
@@ -167,8 +167,6 @@ UCLASS()
 class MESHMODELINGTOOLS_API UDrawPolygonTool : public UInteractiveTool, public IClickSequenceBehaviorTarget
 {
 	GENERATED_BODY()
-	using FVector3d = UE::Geometry::FVector3d;
-	using FVector2d = UE::Geometry::FVector2d;
 public:
 	UDrawPolygonTool();
 
@@ -338,15 +336,14 @@ protected:
 class MESHMODELINGTOOLS_API FDrawPolygonStateChange : public FToolCommandChange
 {
 public:
-	using Points = TArray<UE::Geometry::FVector3d>;
 	bool bHaveDoneUndo = false;
 	int32 CurveTimestamp = 0;
-	const Points FixedVertexPoints;
-	const Points PolyPoints;
+	const TArray<FVector3d> FixedVertexPoints;
+	const TArray<FVector3d> PolyPoints;
 
 	FDrawPolygonStateChange(int32 CurveTimestampIn,
-							const Points& FixedVertexPointsIn,
-							const Points& PolyPointsIn)
+							const TArray<FVector3d>& FixedVertexPointsIn,
+							const TArray<FVector3d>& PolyPointsIn)
 		: CurveTimestamp(CurveTimestampIn),
 		FixedVertexPoints(FixedVertexPointsIn),
 		PolyPoints(PolyPointsIn)
diff --git a/Engine/Plugins/Experimental/MeshModelingToolset/Source/MeshModelingTools/Public/EdgeLoopInsertionTool.h b/Engine/Plugins/Experimental/MeshModelingToolset/Source/MeshModelingTools/Public/EdgeLoopInsertionTool.h
--- a/Engine/Plugins/Experimental/MeshModelingToolset/Source/MeshModelingTools/Public/EdgeLoopInsertionTool.h
+++ b/Engine/Plugins/Experimental/MeshModelingToolset/Source/MeshModelingTools/Public/EdgeLoopInsertionTool.h
@@ -118,7 +118,6 @@ UCLASS()
 class MESHMODELINGTOOLS_API UEdgeLoopInsertionTool : public USingleSelectionTool, public IHoverBehaviorTarget, public IClickBehaviorTarget
 {
 	GENERATED_BODY()
-	using FVector3d = UE::Geometry::FVector3d;
 public:
 
 	friend class UEdgeLoopInsertionOperatorFactory;
diff --git a/Engine/Plugins/Experimental/MeshModelingToolset/Source/MeshModelingTools/Public/EditMeshPolygonsTool.h b/Engine/Plugins/Experimental/MeshModelingToolset/Source/MeshModelingTools/Public/EditMeshPolygonsTool.h
--- a/Engine/Plugins/Experimental/MeshModelingToolset/Source/MeshModelingTools/Public/EditMeshPolygonsTool.h
+++ b/Engine/Plugins/Experimental/MeshModelingToolset/Source/MeshModelingTools/Public/EditMeshPolygonsTool.h
@@ -476,7 +476,6 @@ UCLASS()
 class MESHMODELINGTOOLS_API UEditMeshPolygonsTool : public UMeshSurfacePointTool, public IClickBehaviorTarget
 {
 	GENERATED_BODY()
-	using FVector3d = UE::Geometry::FVector3d;
 	using FTransform3d = UE::Geometry::FTransform3d;
 	using FFrame3d = UE::Geometry::FFrame3d;
 	using FAxisAlignedBox3d = UE::Geometry::FAxisAlignedBox3d;
diff --git a/Engine/Plugins/Experimental/MeshModelingToolset/Source/MeshModelingTools/Public/EditUVIslandsTool.h b/Engine/Plugins/Experimental/MeshModelingToolset/Source/MeshModelingTools/Public/EditUVIslandsTool.h
--- a/Engine/Plugins/Experimental/MeshModelingToolset/Source/MeshModelingTools/Public/EditUVIslandsTool.h
+++ b/Engine/Plugins/Experimental/MeshModelingToolset/Source/MeshModelingTools/Public/EditUVIslandsTool.h
@@ -124,7 +124,7 @@ protected:
 
 	double UVTranslateScale;
 	UE::Geometry::FFrame3d InitialGizmoFrame;
-	UE::Geometry::FVector3d InitialGizmoScale;
+	FVector3d InitialGizmoScale;
 	void ComputeUpdate_Gizmo();
 
 	FUVGroupTopology Topology;
@@ -144,8 +144,8 @@ protected:
 		TArray<int32> Triangles;
 		TArray<int32> UVs;
 		UE::Geometry::FAxisAlignedBox2d UVBounds;
-		UE::Geometry::FVector2d UVOrigin;
-		TArray<UE::Geometry::FVector2f> InitialPositions;
+		FVector2d UVOrigin;
+		TArray<FVector2f> InitialPositions;
 	};
 	TArray<FEditIsland> ActiveIslands;
 	void UpdateUVTransformFromSelection(const FGroupTopologySelection& Selection);
diff --git a/Engine/Plugins/Experimental/MeshModelingToolset/Source/MeshModelingTools/Public/FindPolygonsAlgorithm.h b/Engine/Plugins/Experimental/MeshModelingToolset/Source/MeshModelingTools/Public/FindPolygonsAlgorithm.h
--- a/Engine/Plugins/Experimental/MeshModelingToolset/Source/MeshModelingTools/Public/FindPolygonsAlgorithm.h
+++ b/Engine/Plugins/Experimental/MeshModelingToolset/Source/MeshModelingTools/Public/FindPolygonsAlgorithm.h
@@ -9,7 +9,6 @@ using UE::Geometry::FDynamicMesh3;
 
 class FFindPolygonsAlgorithm
 {
-	using FVector3d = UE::Geometry::FVector3d;
 public:
 
 	FFindPolygonsAlgorithm() {}
diff --git a/Engine/Plugins/Experimental/MeshModelingToolset/Source/MeshModelingTools/Public/GroupEdgeInsertionTool.h b/Engine/Plugins/Experimental/MeshModelingToolset/Source/MeshModelingTools/Public/GroupEdgeInsertionTool.h
--- a/Engine/Plugins/Experimental/MeshModelingToolset/Source/MeshModelingTools/Public/GroupEdgeInsertionTool.h
+++ b/Engine/Plugins/Experimental/MeshModelingToolset/Source/MeshModelingTools/Public/GroupEdgeInsertionTool.h
@@ -90,7 +90,6 @@ class MESHMODELINGTOOLS_API UGroupEdgeInsertionTool : public USingleSelectionToo
 	};
 
 	using FRay3d = UE::Geometry::FRay3d;
-	using FVector3d = UE::Geometry::FVector3d;
 public:
 
 	friend class UGroupEdgeInsertionOperatorFactory;
diff --git a/Engine/Plugins/Experimental/MeshModelingToolset/Source/MeshModelingTools/Public/LatticeDeformerTool.h b/Engine/Plugins/Experimental/MeshModelingToolset/Source/MeshModelingTools/Public/LatticeDeformerTool.h
--- a/Engine/Plugins/Experimental/MeshModelingToolset/Source/MeshModelingTools/Public/LatticeDeformerTool.h
+++ b/Engine/Plugins/Experimental/MeshModelingToolset/Source/MeshModelingTools/Public/LatticeDeformerTool.h
@@ -145,7 +145,7 @@ protected:
 	bool bShouldRebuild = false;
 
 	// Create and store an FFFDLattice. Pass out the lattice's positions and edges.
-	void InitializeLattice(TArray<UE::Geometry::FVector3d>& OutLatticePoints, TArray<UE::Geometry::FVector2i>& OutLatticeEdges);
+	void InitializeLattice(TArray<FVector3d>& OutLatticePoints, TArray<UE::Geometry::FVector2i>& OutLatticeEdges);
 
 	void StartPreview();
 
diff --git a/Engine/Plugins/Experimental/MeshModelingToolset/Source/MeshModelingTools/Public/MeshSpaceDeformerTool.h b/Engine/Plugins/Experimental/MeshModelingToolset/Source/MeshModelingTools/Public/MeshSpaceDeformerTool.h
--- a/Engine/Plugins/Experimental/MeshModelingToolset/Source/MeshModelingTools/Public/MeshSpaceDeformerTool.h
+++ b/Engine/Plugins/Experimental/MeshModelingToolset/Source/MeshModelingTools/Public/MeshSpaceDeformerTool.h
@@ -273,7 +273,7 @@ protected:
 
 	// Button click support
 	EMeshSpaceDeformerToolAction PendingAction;
-	UE::Geometry::FVector3d MeshCenter;
+	FVector3d MeshCenter;
 
 	UE::Geometry::FFrame3d GizmoFrame;
 
@@ -283,7 +283,7 @@ protected:
 
 	TPimplPtr<FSelectClickedAction> SetPointInWorldConnector;
 
-	TArray<UE::Geometry::FVector3d> VisualizationPoints;
+	TArray<FVector3d> VisualizationPoints;
 	FToolDataVisualizer VisualizationRenderer;
 
 	void TransformProxyChanged(UTransformProxy* Proxy, FTransform Transform);
diff --git a/Engine/Plugins/Experimental/MeshModelingToolset/Source/MeshModelingTools/Public/MeshVertexSculptTool.h b/Engine/Plugins/Experimental/MeshModelingToolset/Source/MeshModelingTools/Public/MeshVertexSculptTool.h
--- a/Engine/Plugins/Experimental/MeshModelingToolset/Source/MeshModelingTools/Public/MeshVertexSculptTool.h
+++ b/Engine/Plugins/Experimental/MeshModelingToolset/Source/MeshModelingTools/Public/MeshVertexSculptTool.h
@@ -182,7 +182,6 @@ UCLASS()
 class MESHMODELINGTOOLS_API UMeshVertexSculptTool : public UMeshSculptToolBase
 {
 	GENERATED_BODY()
-	using FVector3d = UE::Geometry::FVector3d;
 public:
 	virtual void Setup() override;
 	virtual void Shutdown(EToolShutdownType ShutdownType) override;
@@ -293,7 +292,7 @@ protected:
 	double SculptMaxFixedHeight = -1.0;
 
 	bool bHaveBrushAlpha = false;
-	UE::Geometry::TImageBuilder<UE::Geometry::FVector4f> BrushAlphaValues;
+	UE::Geometry::TImageBuilder<FVector4f> BrushAlphaValues;
 	UE::Geometry::FImageDimensions BrushAlphaDimensions;
 	double SampleBrushAlpha(const FSculptBrushStamp& Stamp, const FVector3d& Position) const;
 
diff --git a/Engine/Plugins/Experimental/MeshModelingToolset/Source/MeshModelingTools/Public/MirrorTool.h b/Engine/Plugins/Experimental/MeshModelingToolset/Source/MeshModelingTools/Public/MirrorTool.h
--- a/Engine/Plugins/Experimental/MeshModelingToolset/Source/MeshModelingTools/Public/MirrorTool.h
+++ b/Engine/Plugins/Experimental/MeshModelingToolset/Source/MeshModelingTools/Public/MirrorTool.h
@@ -188,7 +188,6 @@ UCLASS()
 class MESHMODELINGTOOLS_API UMirrorTool : public UMultiSelectionTool, public IModifierToggleBehaviorTarget
 {
 	GENERATED_BODY()
-	using FVector3d = UE::Geometry::FVector3d;
 public:
 
 	friend UMirrorOperatorFactory;
diff --git a/Engine/Plugins/Experimental/MeshModelingToolset/Source/MeshModelingTools/Public/Properties/RevolveProperties.h b/Engine/Plugins/Experimental/MeshModelingToolset/Source/MeshModelingTools/Public/Properties/RevolveProperties.h
--- a/Engine/Plugins/Experimental/MeshModelingToolset/Source/MeshModelingTools/Public/Properties/RevolveProperties.h
+++ b/Engine/Plugins/Experimental/MeshModelingToolset/Source/MeshModelingTools/Public/Properties/RevolveProperties.h
@@ -59,7 +59,6 @@ UCLASS()
 class MESHMODELINGTOOLS_API URevolveProperties : public UInteractiveToolPropertySet
 {
 	GENERATED_BODY()
-	using FVector3d = UE::Geometry::FVector3d;
 public:
 
 	URevolveProperties()
diff --git a/Engine/Plugins/Experimental/MeshModelingToolset/Source/MeshModelingTools/Public/RevolveBoundaryTool.h b/Engine/Plugins/Experimental/MeshModelingToolset/Source/MeshModelingTools/Public/RevolveBoundaryTool.h
--- a/Engine/Plugins/Experimental/MeshModelingToolset/Source/MeshModelingTools/Public/RevolveBoundaryTool.h
+++ b/Engine/Plugins/Experimental/MeshModelingToolset/Source/MeshModelingTools/Public/RevolveBoundaryTool.h
@@ -84,7 +84,6 @@ UCLASS()
 class MESHMODELINGTOOLS_API URevolveBoundaryTool : public UMeshBoundaryToolBase, public IClickBehaviorTarget
 {
 	GENERATED_BODY()
-	using FVector3d = UE::Geometry::FVector3d;
 
 public:
 	virtual void SetAssetAPI(IAssetGenerationAPI* NewAssetApi) { AssetAPI = NewAssetApi; }
diff --git a/Engine/Plugins/Experimental/MeshModelingToolset/Source/MeshModelingTools/Public/Sculpting/MeshBrushOpBase.h b/Engine/Plugins/Experimental/MeshModelingToolset/Source/MeshModelingTools/Public/Sculpting/MeshBrushOpBase.h
--- a/Engine/Plugins/Experimental/MeshModelingToolset/Source/MeshModelingTools/Public/Sculpting/MeshBrushOpBase.h
+++ b/Engine/Plugins/Experimental/MeshModelingToolset/Source/MeshModelingTools/Public/Sculpting/MeshBrushOpBase.h
@@ -11,7 +11,6 @@
 PREDECLARE_USE_GEOMETRY_CLASS(FDynamicMesh3);
 using UE::Geometry::FFrame3d;
 using UE::Geometry::FMatrix3d;
-using UE::Geometry::FVector3d;
 
 enum class ESculptBrushOpTargetType : uint8
 {
diff --git a/Engine/Plugins/Experimental/MeshModelingToolset/Source/MeshModelingTools/Public/Sculpting/MeshSculptToolBase.h b/Engine/Plugins/Experimental/MeshModelingToolset/Source/MeshModelingTools/Public/Sculpting/MeshSculptToolBase.h
--- a/Engine/Plugins/Experimental/MeshModelingToolset/Source/MeshModelingTools/Public/Sculpting/MeshSculptToolBase.h
+++ b/Engine/Plugins/Experimental/MeshModelingToolset/Source/MeshModelingTools/Public/Sculpting/MeshSculptToolBase.h
@@ -167,10 +167,6 @@ class MESHMODELINGTOOLS_API UMeshSculptToolBase : public UMeshSurfacePointTool
 {
 	GENERATED_BODY()
 protected:
-	using FVector3d = UE::Geometry::FVector3d;
-	using FVector3f = UE::Geometry::FVector3f;
-	using FVector2d = UE::Geometry::FVector2d;
-	using FVector2f = UE::Geometry::FVector2f;
 	using FTransform3d = UE::Geometry::FTransform3d;
 	using FFrame3d = UE::Geometry::FFrame3d;
 	using FRay3d = UE::Geometry::FRay3d;
diff --git a/Engine/Plugins/Experimental/MeshModelingToolset/Source/MeshModelingTools/Public/SeamSculptTool.h b/Engine/Plugins/Experimental/MeshModelingToolset/Source/MeshModelingTools/Public/SeamSculptTool.h
--- a/Engine/Plugins/Experimental/MeshModelingToolset/Source/MeshModelingTools/Public/SeamSculptTool.h
+++ b/Engine/Plugins/Experimental/MeshModelingToolset/Source/MeshModelingTools/Public/SeamSculptTool.h
@@ -50,7 +50,6 @@ UCLASS(Transient)
 class MESHMODELINGTOOLS_API USeamSculptTool : public UDynamicMeshBrushTool
 {
 	GENERATED_BODY()
-	using FVector3d = UE::Geometry::FVector3d;
 public:
 	USeamSculptTool();
 
diff --git a/Engine/Plugins/Experimental/MeshModelingToolset/Source/MeshModelingToolsEditorOnly/Public/MeshTangentsTool.h b/Engine/Plugins/Experimental/MeshModelingToolset/Source/MeshModelingToolsEditorOnly/Public/MeshTangentsTool.h
--- a/Engine/Plugins/Experimental/MeshModelingToolset/Source/MeshModelingToolsEditorOnly/Public/MeshTangentsTool.h
+++ b/Engine/Plugins/Experimental/MeshModelingToolset/Source/MeshModelingToolsEditorOnly/Public/MeshTangentsTool.h
@@ -81,7 +81,6 @@ UCLASS()
 class MESHMODELINGTOOLSEDITORONLY_API UMeshTangentsTool : public USingleSelectionTool, public UE::Geometry::IGenericDataOperatorFactory<UE::Geometry::FMeshTangentsd>
 {
 	GENERATED_BODY()
-	using FVector3d = UE::Geometry::FVector3d;
 public:
 	UMeshTangentsTool();
 
diff --git a/Engine/Plugins/Experimental/MeshModelingToolset/Source/ModelingComponents/Private/BaseDynamicMeshSceneProxy.h b/Engine/Plugins/Experimental/MeshModelingToolset/Source/ModelingComponents/Private/BaseDynamicMeshSceneProxy.h
--- a/Engine/Plugins/Experimental/MeshModelingToolset/Source/ModelingComponents/Private/BaseDynamicMeshSceneProxy.h
+++ b/Engine/Plugins/Experimental/MeshModelingToolset/Source/ModelingComponents/Private/BaseDynamicMeshSceneProxy.h
@@ -284,9 +284,6 @@ protected:
  */
 class FBaseDynamicMeshSceneProxy : public FPrimitiveSceneProxy
 {
-	using FVector3d = UE::Geometry::FVector3d;
-	using FVector3f = UE::Geometry::FVector3f;
-	using FVector2f = UE::Geometry::FVector2f;
 	using FIndex2i = UE::Geometry::FIndex2i;
 	using FIndex3i = UE::Geometry::FIndex3i;
 public:
diff --git a/Engine/Plugins/Experimental/MeshModelingToolset/Source/ModelingComponents/Private/OctreeDynamicMeshSceneProxy.h b/Engine/Plugins/Experimental/MeshModelingToolset/Source/ModelingComponents/Private/OctreeDynamicMeshSceneProxy.h
--- a/Engine/Plugins/Experimental/MeshModelingToolset/Source/ModelingComponents/Private/OctreeDynamicMeshSceneProxy.h
+++ b/Engine/Plugins/Experimental/MeshModelingToolset/Source/ModelingComponents/Private/OctreeDynamicMeshSceneProxy.h
@@ -29,9 +29,6 @@ DECLARE_CYCLE_STAT(TEXT("SculptToolOctree_BufferUpload"), STAT_SculptToolOctree_
  */
 class FOctreeDynamicMeshSceneProxy final : public FBaseDynamicMeshSceneProxy
 {
-	using FVector3d = UE::Geometry::FVector3d;
-	using FVector3f = UE::Geometry::FVector3f;
-	using FVector2f = UE::Geometry::FVector2f;
 	using FIndex2i = UE::Geometry::FIndex2i;
 	using FIndex3i = UE::Geometry::FIndex3i;
 private:
diff --git a/Engine/Plugins/Experimental/MeshModelingToolset/Source/ModelingComponents/Private/SimpleDynamicMeshSceneProxy.h b/Engine/Plugins/Experimental/MeshModelingToolset/Source/ModelingComponents/Private/SimpleDynamicMeshSceneProxy.h
--- a/Engine/Plugins/Experimental/MeshModelingToolset/Source/ModelingComponents/Private/SimpleDynamicMeshSceneProxy.h
+++ b/Engine/Plugins/Experimental/MeshModelingToolset/Source/ModelingComponents/Private/SimpleDynamicMeshSceneProxy.h
@@ -16,9 +16,6 @@
  */
 class FSimpleDynamicMeshSceneProxy final : public FBaseDynamicMeshSceneProxy
 {
-	using FVector3d = UE::Geometry::FVector3d;
-	using FVector3f = UE::Geometry::FVector3f;
-	using FVector2f = UE::Geometry::FVector2f;
 	using FIndex2i = UE::Geometry::FIndex2i;
 	using FIndex3i = UE::Geometry::FIndex3i;
 private:
diff --git a/Engine/Plugins/Experimental/MeshModelingToolset/Source/ModelingComponents/Public/BaseTools/BaseCreateFromSelectedTool.h b/Engine/Plugins/Experimental/MeshModelingToolset/Source/ModelingComponents/Public/BaseTools/BaseCreateFromSelectedTool.h
--- a/Engine/Plugins/Experimental/MeshModelingToolset/Source/ModelingComponents/Public/BaseTools/BaseCreateFromSelectedTool.h
+++ b/Engine/Plugins/Experimental/MeshModelingToolset/Source/ModelingComponents/Public/BaseTools/BaseCreateFromSelectedTool.h
@@ -100,10 +100,6 @@ class MODELINGCOMPONENTS_API UBaseCreateFromSelectedTool : public UMultiSelectio
 {
 	GENERATED_BODY()
 protected:
-	using FVector3d = UE::Geometry::FVector3d;
-	using FVector3f = UE::Geometry::FVector3f;
-	using FVector2d = UE::Geometry::FVector2d;
-	using FVector2f = UE::Geometry::FVector2f;
 	using FTransform3d = UE::Geometry::FTransform3d;
 	using FFrame3d = UE::Geometry::FFrame3d;
 	using FRay3d = UE::Geometry::FRay3d;
diff --git a/Engine/Plugins/Experimental/MeshModelingToolset/Source/ModelingComponents/Public/BaseTools/BaseMeshProcessingTool.h b/Engine/Plugins/Experimental/MeshModelingToolset/Source/ModelingComponents/Public/BaseTools/BaseMeshProcessingTool.h
--- a/Engine/Plugins/Experimental/MeshModelingToolset/Source/ModelingComponents/Public/BaseTools/BaseMeshProcessingTool.h
+++ b/Engine/Plugins/Experimental/MeshModelingToolset/Source/ModelingComponents/Public/BaseTools/BaseMeshProcessingTool.h
@@ -79,10 +79,6 @@ class MODELINGCOMPONENTS_API UBaseMeshProcessingTool : public USingleSelectionTo
 {
 	GENERATED_BODY()
 protected:
-	using FVector3d = UE::Geometry::FVector3d;
-	using FVector3f = UE::Geometry::FVector3f;
-	using FVector2d = UE::Geometry::FVector2d;
-	using FVector2f = UE::Geometry::FVector2f;
 	using FTransform3d = UE::Geometry::FTransform3d;
 	using FFrame3d = UE::Geometry::FFrame3d;
 	using FRay3d = UE::Geometry::FRay3d;
diff --git a/Engine/Plugins/Experimental/MeshModelingToolset/Source/ModelingComponents/Public/Changes/MeshVertexChange.h b/Engine/Plugins/Experimental/MeshModelingToolset/Source/ModelingComponents/Public/Changes/MeshVertexChange.h
--- a/Engine/Plugins/Experimental/MeshModelingToolset/Source/ModelingComponents/Public/Changes/MeshVertexChange.h
+++ b/Engine/Plugins/Experimental/MeshModelingToolset/Source/ModelingComponents/Public/Changes/MeshVertexChange.h
@@ -19,8 +19,6 @@ PREDECLARE_USE_GEOMETRY_CLASS(FDynamicMesh3);
  */
 class MODELINGCOMPONENTS_API FMeshVertexChange : public FToolCommandChange
 {
-	using FVector3d = UE::Geometry::FVector3d;
-	using FVector3f = UE::Geometry::FVector3f;
 public:
 	bool bHaveVertexPositions = true;
 	bool bHaveVertexColors = false;
@@ -61,8 +59,6 @@ ENUM_CLASS_FLAGS(EMeshVertexChangeComponents);
  */
 class MODELINGCOMPONENTS_API FMeshVertexChangeBuilder
 {
-	using FVector3d = UE::Geometry::FVector3d;
-	using FVector3f = UE::Geometry::FVector3f;
 public:
 	TUniquePtr<FMeshVertexChange> Change;
 	TMap<int32, int32> SavedVertices;
diff --git a/Engine/Plugins/Experimental/MeshModelingToolset/Source/ModelingComponents/Public/Drawing/PolyEditPreviewMesh.h b/Engine/Plugins/Experimental/MeshModelingToolset/Source/ModelingComponents/Public/Drawing/PolyEditPreviewMesh.h
--- a/Engine/Plugins/Experimental/MeshModelingToolset/Source/ModelingComponents/Public/Drawing/PolyEditPreviewMesh.h
+++ b/Engine/Plugins/Experimental/MeshModelingToolset/Source/ModelingComponents/Public/Drawing/PolyEditPreviewMesh.h
@@ -18,8 +18,6 @@ UCLASS(Transient)
 class MODELINGCOMPONENTS_API UPolyEditPreviewMesh : public UPreviewMesh
 {
 	GENERATED_BODY()
-	using FVector3d = UE::Geometry::FVector3d;
-	using FVector3f = UE::Geometry::FVector3f;
 	using FTransform3d = UE::Geometry::FTransform3d;	
 public:
 
diff --git a/Engine/Plugins/Experimental/MeshModelingToolset/Source/ModelingComponents/Public/Mechanics/CollectSurfacePathMechanic.h b/Engine/Plugins/Experimental/MeshModelingToolset/Source/ModelingComponents/Public/Mechanics/CollectSurfacePathMechanic.h
--- a/Engine/Plugins/Experimental/MeshModelingToolset/Source/ModelingComponents/Public/Mechanics/CollectSurfacePathMechanic.h
+++ b/Engine/Plugins/Experimental/MeshModelingToolset/Source/ModelingComponents/Public/Mechanics/CollectSurfacePathMechanic.h
@@ -31,7 +31,6 @@ class MODELINGCOMPONENTS_API UCollectSurfacePathMechanic : public UInteractionMe
 public:
 	using FFrame3d = UE::Geometry::FFrame3d;
 	using FRay3d = UE::Geometry::FRay3d;
-	using FVector3d = UE::Geometry::FVector3d;
 
 	TUniqueFunction<bool()> IsDoneFunc = nullptr;
 
diff --git a/Engine/Plugins/Experimental/MeshModelingToolset/Source/ModelingComponents/Public/Mechanics/ConstructionPlaneMechanic.h b/Engine/Plugins/Experimental/MeshModelingToolset/Source/ModelingComponents/Public/Mechanics/ConstructionPlaneMechanic.h
--- a/Engine/Plugins/Experimental/MeshModelingToolset/Source/ModelingComponents/Public/Mechanics/ConstructionPlaneMechanic.h
+++ b/Engine/Plugins/Experimental/MeshModelingToolset/Source/ModelingComponents/Public/Mechanics/ConstructionPlaneMechanic.h
@@ -47,7 +47,7 @@ public:
 
 	void Initialize(UWorld* TargetWorld, const UE::Geometry::FFrame3d& InitialPlane);
 
-	void SetDrawPlaneFromWorldPos(const UE::Geometry::FVector3d& Position, const UE::Geometry::FVector3d& Normal, bool bIgnoreNormal);
+	void SetDrawPlaneFromWorldPos(const FVector3d& Position, const FVector3d& Normal, bool bIgnoreNormal);
 
 	/** 
 	 * Sets the plane without broadcasting OnPlaneChanged. Useful when the user of the tool wants to change
diff --git a/Engine/Plugins/Experimental/MeshModelingToolset/Source/ModelingComponents/Public/Mechanics/CurveControlPointsMechanic.h b/Engine/Plugins/Experimental/MeshModelingToolset/Source/ModelingComponents/Public/Mechanics/CurveControlPointsMechanic.h
--- a/Engine/Plugins/Experimental/MeshModelingToolset/Source/ModelingComponents/Public/Mechanics/CurveControlPointsMechanic.h
+++ b/Engine/Plugins/Experimental/MeshModelingToolset/Source/ModelingComponents/Public/Mechanics/CurveControlPointsMechanic.h
@@ -45,7 +45,6 @@ class MODELINGCOMPONENTS_API UCurveControlPointsMechanic : public UInteractionMe
 	GENERATED_BODY()
 
 protected:
-	using FVector3d = UE::Geometry::FVector3d;
 
 	// We want some way to store the control point sequence that lets us easily associate points with their renderable and hit-testable
 	// representations, since we need to alter all of these together as points get moved or added. We use FOrderedPoints for this, until
@@ -457,7 +456,6 @@ protected:
 
 class MODELINGCOMPONENTS_API FCurveControlPointsMechanicInsertionChange : public FToolCommandChange
 {
-	using FVector3d = UE::Geometry::FVector3d;
 public:
 	FCurveControlPointsMechanicInsertionChange(int32 SequencePositionIn, int32 PointID, 
 		const FVector3d& CoordinatesIn, bool AddedIn, int32 ChangeStampIn);
@@ -499,7 +497,6 @@ protected:
 
 class MODELINGCOMPONENTS_API FCurveControlPointsMechanicMovementChange : public FToolCommandChange
 {
-	using FVector3d = UE::Geometry::FVector3d;
 public:
 	FCurveControlPointsMechanicMovementChange(int32 PointIDIn, const FVector3d& OriginalPositionIn, 
 		const FVector3d& NewPositionIn, int32 ChangeStampIn);
diff --git a/Engine/Plugins/Experimental/MeshModelingToolset/Source/ModelingComponents/Public/Mechanics/DragAlignmentMechanic.h b/Engine/Plugins/Experimental/MeshModelingToolset/Source/ModelingComponents/Public/Mechanics/DragAlignmentMechanic.h
--- a/Engine/Plugins/Experimental/MeshModelingToolset/Source/ModelingComponents/Public/Mechanics/DragAlignmentMechanic.h
+++ b/Engine/Plugins/Experimental/MeshModelingToolset/Source/ModelingComponents/Public/Mechanics/DragAlignmentMechanic.h
@@ -90,7 +90,7 @@ protected:
 	bool bPreviewEndpointsValid = false;
 	bool bWaitingOnProjectedResult = false;
 	FHitResult LastHitResult;
-	UE::Geometry::FVector3d LastProjectedResult;
+	FVector3d LastProjectedResult;
 
 	bool bAlignmentToggle = false;
 	int32 AlignmentModifierID = 1;
diff --git a/Engine/Plugins/Experimental/MeshModelingToolset/Source/ModelingComponents/Public/Mechanics/LatticeControlPointsMechanic.h b/Engine/Plugins/Experimental/MeshModelingToolset/Source/ModelingComponents/Public/Mechanics/LatticeControlPointsMechanic.h
--- a/Engine/Plugins/Experimental/MeshModelingToolset/Source/ModelingComponents/Public/Mechanics/LatticeControlPointsMechanic.h
+++ b/Engine/Plugins/Experimental/MeshModelingToolset/Source/ModelingComponents/Public/Mechanics/LatticeControlPointsMechanic.h
@@ -31,7 +31,6 @@ class MODELINGCOMPONENTS_API ULatticeControlPointsMechanic :
 	public UInteractionMechanic, public IClickBehaviorTarget, public IHoverBehaviorTarget
 {
 	GENERATED_BODY()
-	using FVector3d = UE::Geometry::FVector3d;
 	using FVector2i = UE::Geometry::FVector2i;
 	using FTransform3d = UE::Geometry::FTransform3d;
 
@@ -211,7 +210,6 @@ protected:
 // Control points have moved
 class MODELINGCOMPONENTS_API FLatticeControlPointsMechanicMovementChange : public FToolCommandChange
 {
-	using FVector3d = UE::Geometry::FVector3d;
 public:
 
 	FLatticeControlPointsMechanicMovementChange(
diff --git a/Engine/Plugins/Experimental/MeshModelingToolset/Source/ModelingComponents/Public/Mechanics/PlaneDistanceFromHitMechanic.h b/Engine/Plugins/Experimental/MeshModelingToolset/Source/ModelingComponents/Public/Mechanics/PlaneDistanceFromHitMechanic.h
--- a/Engine/Plugins/Experimental/MeshModelingToolset/Source/ModelingComponents/Public/Mechanics/PlaneDistanceFromHitMechanic.h
+++ b/Engine/Plugins/Experimental/MeshModelingToolset/Source/ModelingComponents/Public/Mechanics/PlaneDistanceFromHitMechanic.h
@@ -18,7 +18,6 @@ UCLASS()
 class MODELINGCOMPONENTS_API UPlaneDistanceFromHitMechanic : public UInteractionMechanic
 {
 	GENERATED_BODY()
-	using FVector3d = UE::Geometry::FVector3d;
 	using FFrame3d = UE::Geometry::FFrame3d;
 public:
 	/** If this function is set, we will check it for a ray intersection if the target mesh is not hit */
diff --git a/Engine/Plugins/Experimental/MeshModelingToolset/Source/ModelingComponents/Public/Mechanics/SpaceCurveDeformationMechanic.h b/Engine/Plugins/Experimental/MeshModelingToolset/Source/ModelingComponents/Public/Mechanics/SpaceCurveDeformationMechanic.h
--- a/Engine/Plugins/Experimental/MeshModelingToolset/Source/ModelingComponents/Public/Mechanics/SpaceCurveDeformationMechanic.h
+++ b/Engine/Plugins/Experimental/MeshModelingToolset/Source/ModelingComponents/Public/Mechanics/SpaceCurveDeformationMechanic.h
@@ -68,7 +68,6 @@ UCLASS()
 class MODELINGCOMPONENTS_API USpaceCurveDeformationMechanicPropertySet : public UInteractiveToolPropertySet
 {
 	GENERATED_BODY()
-	using FVector3d = UE::Geometry::FVector3d;
 
 public:
 	UPROPERTY(EditAnywhere, Category = Options)
@@ -95,7 +94,6 @@ UCLASS()
 class MODELINGCOMPONENTS_API USpaceCurveDeformationMechanic : public UInteractionMechanic, public IClickBehaviorTarget, public IHoverBehaviorTarget
 {
 	GENERATED_BODY()
-	using FVector3d = UE::Geometry::FVector3d;
 public:
 
 	// Behaviors used for moving points around and hovering them
diff --git a/Engine/Plugins/Experimental/MeshModelingToolset/Source/ModelingComponents/Public/Mechanics/SpatialCurveDistanceMechanic.h b/Engine/Plugins/Experimental/MeshModelingToolset/Source/ModelingComponents/Public/Mechanics/SpatialCurveDistanceMechanic.h
--- a/Engine/Plugins/Experimental/MeshModelingToolset/Source/ModelingComponents/Public/Mechanics/SpatialCurveDistanceMechanic.h
+++ b/Engine/Plugins/Experimental/MeshModelingToolset/Source/ModelingComponents/Public/Mechanics/SpatialCurveDistanceMechanic.h
@@ -15,7 +15,6 @@ UCLASS()
 class MODELINGCOMPONENTS_API USpatialCurveDistanceMechanic : public UInteractionMechanic
 {
 	GENERATED_BODY()
-	using FVector3d = UE::Geometry::FVector3d;
 	using FTransform3d = UE::Geometry::FTransform3d;
 public:
 	/** If this function is set, the hit point will be passed in to this function for snapping. Return false to indicate no snapping occurred. */
diff --git a/Engine/Plugins/Experimental/MeshModelingToolset/Source/ModelingComponents/Public/PreviewMesh.h b/Engine/Plugins/Experimental/MeshModelingToolset/Source/ModelingComponents/Public/PreviewMesh.h
--- a/Engine/Plugins/Experimental/MeshModelingToolset/Source/ModelingComponents/Public/PreviewMesh.h
+++ b/Engine/Plugins/Experimental/MeshModelingToolset/Source/ModelingComponents/Public/PreviewMesh.h
@@ -58,7 +58,6 @@ class MODELINGCOMPONENTS_API UPreviewMesh : public UObject, public IMeshVertexCo
 {
 	GENERATED_BODY()
 	using FRay3d = UE::Geometry::FRay3d;
-	using FVector3d = UE::Geometry::FVector3d;
 public:
 	UPreviewMesh();
 	virtual ~UPreviewMesh();
diff --git a/Engine/Plugins/Experimental/MeshModelingToolset/Source/ModelingComponents/Public/Selection/GroupTopologySelector.h b/Engine/Plugins/Experimental/MeshModelingToolset/Source/ModelingComponents/Public/Selection/GroupTopologySelector.h
--- a/Engine/Plugins/Experimental/MeshModelingToolset/Source/ModelingComponents/Public/Selection/GroupTopologySelector.h
+++ b/Engine/Plugins/Experimental/MeshModelingToolset/Source/ModelingComponents/Public/Selection/GroupTopologySelector.h
@@ -36,7 +36,6 @@ using UE::Geometry::FGeometrySet3;
 class MODELINGCOMPONENTS_API FGroupTopologySelector
 {
 	using FRay3d = UE::Geometry::FRay3d;
-	using FVector3d = UE::Geometry::FVector3d;
 	using FTransform3d = UE::Geometry::FTransform3d;
 public:
 
diff --git a/Engine/Plugins/Experimental/MeshModelingToolset/Source/ModelingComponents/Public/Selection/PolygonSelectionMechanic.h b/Engine/Plugins/Experimental/MeshModelingToolset/Source/ModelingComponents/Public/Selection/PolygonSelectionMechanic.h
--- a/Engine/Plugins/Experimental/MeshModelingToolset/Source/ModelingComponents/Public/Selection/PolygonSelectionMechanic.h
+++ b/Engine/Plugins/Experimental/MeshModelingToolset/Source/ModelingComponents/Public/Selection/PolygonSelectionMechanic.h
@@ -82,7 +82,6 @@ class MODELINGCOMPONENTS_API UPolygonSelectionMechanic : public UInteractionMech
 {
 	GENERATED_BODY()
 	using FFrame3d = UE::Geometry::FFrame3d;
-	using FVector3d = UE::Geometry::FVector3d;
 	using FTransform3d = UE::Geometry::FTransform3d;	
 public:
 
diff --git a/Engine/Plugins/Experimental/MeshModelingToolset/Source/ModelingComponents/Public/ToolSceneQueriesUtil.h b/Engine/Plugins/Experimental/MeshModelingToolset/Source/ModelingComponents/Public/ToolSceneQueriesUtil.h
--- a/Engine/Plugins/Experimental/MeshModelingToolset/Source/ModelingComponents/Public/ToolSceneQueriesUtil.h
+++ b/Engine/Plugins/Experimental/MeshModelingToolset/Source/ModelingComponents/Public/ToolSceneQueriesUtil.h
@@ -14,7 +14,6 @@
 namespace ToolSceneQueriesUtil
 {
 	using namespace UE::Geometry;
-	using FVector3d = UE::Geometry::FVector3d;
 
 	/**
 	 * @return global visual angle snap threshold (default is 1 degree)
diff --git a/Engine/Plugins/Experimental/MeshModelingToolset/Source/ModelingComponents/Public/Transforms/MultiTransformer.h b/Engine/Plugins/Experimental/MeshModelingToolset/Source/ModelingComponents/Public/Transforms/MultiTransformer.h
--- a/Engine/Plugins/Experimental/MeshModelingToolset/Source/ModelingComponents/Public/Transforms/MultiTransformer.h
+++ b/Engine/Plugins/Experimental/MeshModelingToolset/Source/ModelingComponents/Public/Transforms/MultiTransformer.h
@@ -31,7 +31,6 @@ class MODELINGCOMPONENTS_API UMultiTransformer : public UObject
 {
 	GENERATED_BODY()
 	using FFrame3d = UE::Geometry::FFrame3d;
-	using FVector3d = UE::Geometry::FVector3d;
 public:
 	
 	virtual void Setup(UInteractiveGizmoManager* GizmoManager, IToolContextTransactionProvider* TransactionProviderIn);
diff --git a/Engine/Plugins/Experimental/PlanarCutPlugin/Source/PlanarCut/Public/FractureAutoUV.h b/Engine/Plugins/Experimental/PlanarCutPlugin/Source/PlanarCut/Public/FractureAutoUV.h
--- a/Engine/Plugins/Experimental/PlanarCutPlugin/Source/PlanarCut/Public/FractureAutoUV.h
+++ b/Engine/Plugins/Experimental/PlanarCutPlugin/Source/PlanarCut/Public/FractureAutoUV.h
@@ -41,7 +41,7 @@ void PLANARCUT_API TextureInternalSurfaces(
 	FGeometryCollection& Collection,
 	double MaxDistance,
 	int32 GutterSize,
-	UE::Geometry::TImageBuilder<UE::Geometry::FVector3f>& TextureOut,
+	UE::Geometry::TImageBuilder<FVector3f>& TextureOut,
 	bool bOnlyOddMaterials = true,
 	TArrayView<int32> WhichMaterials = TArrayView<int32>()
 );
diff --git a/Engine/Plugins/Experimental/Text3D/Source/Text3D/Private/Arrangement.h b/Engine/Plugins/Experimental/Text3D/Source/Text3D/Private/Arrangement.h
--- a/Engine/Plugins/Experimental/Text3D/Source/Text3D/Private/Arrangement.h
+++ b/Engine/Plugins/Experimental/Text3D/Source/Text3D/Private/Arrangement.h
@@ -10,7 +10,6 @@
  */
 struct FArrangement final
 {
-	using FVector2d = UE::Geometry::FVector2d;
 	using FIndex2i = UE::Geometry::FIndex2i;
 	using FSegment2f = UE::Geometry::FSegment2f;
 	using FSegment2d = UE::Geometry::FSegment2d;
diff --git a/Engine/Plugins/Experimental/Text3D/Source/Text3D/Private/GlyphLoader.h b/Engine/Plugins/Experimental/Text3D/Source/Text3D/Private/GlyphLoader.h
--- a/Engine/Plugins/Experimental/Text3D/Source/Text3D/Private/GlyphLoader.h
+++ b/Engine/Plugins/Experimental/Text3D/Source/Text3D/Private/GlyphLoader.h
@@ -14,7 +14,6 @@ using UE::Geometry::FPolygon2f;
 
 class FGlyphLoader final
 {
-	using FVector2f = UE::Geometry::FVector2f;
 public:
 	FGlyphLoader(const FT_GlyphSlot Glyph);
 
diff --git a/Engine/Plugins/Experimental/Text3D/Source/Text3D/Private/MeshCreator.cpp b/Engine/Plugins/Experimental/Text3D/Source/Text3D/Private/MeshCreator.cpp
--- a/Engine/Plugins/Experimental/Text3D/Source/Text3D/Private/MeshCreator.cpp
+++ b/Engine/Plugins/Experimental/Text3D/Source/Text3D/Private/MeshCreator.cpp
@@ -370,7 +370,7 @@ void FMeshCreator::TriangulateAndConvert(const TSharedContourNode& Node, int32&
 			FContour& Contour = ContoursLocal->Add();
 			const FPolygon2f& Polygon = *NodeIn->Contour;
 
-			for (const UE::Geometry::FVector2f& Vertex : Polygon.GetVertices())
+			for (const FVector2f& Vertex : Polygon.GetVertices())
 			{
 				// Add point to contour in old format
 				const FPartPtr Point = MakeShared<FPart>();
diff --git a/Engine/Plugins/MeshPainting/Source/MeshPaintingToolset/Private/BaseMeshPaintComponentAdapter.cpp b/Engine/Plugins/MeshPainting/Source/MeshPaintingToolset/Private/BaseMeshPaintComponentAdapter.cpp
--- a/Engine/Plugins/MeshPainting/Source/MeshPaintingToolset/Private/BaseMeshPaintComponentAdapter.cpp
+++ b/Engine/Plugins/MeshPainting/Source/MeshPaintingToolset/Private/BaseMeshPaintComponentAdapter.cpp
@@ -186,9 +186,9 @@ bool FBaseMeshPaintComponentAdapter::RayIntersectAdapter(UE::Geometry::FIndex3i&
 {
 	int32 HitTriangleID;
 	double NearestT;
-	UE::Geometry::FVector3d Direction((UE::Geometry::FVector3d)End - (UE::Geometry::FVector3d)Start);
+	FVector3d Direction((FVector3d)End - (FVector3d)Start);
 	UE::Geometry::Normalize(Direction);
-	UE::Geometry::FRay3d Ray((UE::Geometry::FVector3d)Start, Direction);
+	UE::Geometry::FRay3d Ray((FVector3d)Start, Direction);
 	bool bHit = AABBTree->FindNearestHitTriangle(Ray, NearestT, HitTriangleID);
 	if (bHit)
 	{
diff --git a/Engine/Plugins/MeshPainting/Source/MeshPaintingToolset/Private/MeshPaintStaticMeshAdapter.cpp b/Engine/Plugins/MeshPainting/Source/MeshPaintingToolset/Private/MeshPaintStaticMeshAdapter.cpp
--- a/Engine/Plugins/MeshPainting/Source/MeshPaintingToolset/Private/MeshPaintStaticMeshAdapter.cpp
+++ b/Engine/Plugins/MeshPainting/Source/MeshPaintingToolset/Private/MeshPaintStaticMeshAdapter.cpp
@@ -212,14 +212,14 @@ bool FMeshPaintStaticMeshComponentAdapter::LineTraceComponent(struct FHitResult&
 		const FVector& P1 = MeshVertices[FoundTriangle.B];
 		const FVector& P2 = MeshVertices[FoundTriangle.C];
 
-		UE::Geometry::FTriangle3d Triangle((UE::Geometry::FVector3d)P0, (UE::Geometry::FVector3d)P1, (UE::Geometry::FVector3d)P2);
-		UE::Geometry::FVector3d TriNormal = Triangle.Normal();
+		UE::Geometry::FTriangle3d Triangle((FVector3d)P0, (FVector3d)P1, (FVector3d)P2);
+		FVector3d TriNormal = Triangle.Normal();
 
 		//check collinearity of A,B,C
 		if (TriNormal.SquaredLength() > (double)SMALL_NUMBER)
 		{
-			UE::Geometry::FVector3d RayDirection = ((UE::Geometry::FVector3d)LocalEnd - (UE::Geometry::FVector3d)LocalStart);
-			UE::Geometry::FRay3d LocalRay((UE::Geometry::FVector3d)LocalStart, UE::Geometry::Normalized(RayDirection));
+			FVector3d RayDirection = ((FVector3d)LocalEnd - (FVector3d)LocalStart);
+			UE::Geometry::FRay3d LocalRay((FVector3d)LocalStart, UE::Geometry::Normalized(RayDirection));
 			UE::Geometry::FIntrRay3Triangle3d RayTriIntersection(LocalRay, Triangle);
 			if (RayTriIntersection.Find())
 			{
