LWC: Build fixes

diff --git a/Engine/Plugins/Experimental/Enterprise/DataprepGeometryOperations/Source/Private/JacketingProcess.cpp b/Engine/Plugins/Experimental/Enterprise/DataprepGeometryOperations/Source/Private/JacketingProcess.cpp
--- a/Engine/Plugins/Experimental/Enterprise/DataprepGeometryOperations/Source/Private/JacketingProcess.cpp
+++ b/Engine/Plugins/Experimental/Enterprise/DataprepGeometryOperations/Source/Private/JacketingProcess.cpp
@@ -266,7 +266,7 @@ void FJacketingProcess::FindOverlappingActors(const TArray<AActor*>& InActorsToT
 
 			const FMeshDescription* MeshDescription = StaticMeshComponent->GetStaticMesh()->GetMeshDescription(0);
 			const FTransform& ComponentTransform = StaticMeshComponent->GetComponentTransform();
-			TVertexAttributesConstRef<FVector> VertexPositions = MeshDescription->VertexAttributes().GetAttributesRef<FVector>(MeshAttribute::Vertex::Position);
+			TVertexAttributesConstRef<FVector3f> VertexPositions = MeshDescription->VertexAttributes().GetAttributesRef<FVector3f>(MeshAttribute::Vertex::Position);
 
 			for (FVertexID VertexID : MeshDescription->Vertices().GetElementIDs())
 			{
@@ -520,7 +520,7 @@ void FJacketingProcess::ApplyJacketingOnMeshActors(const TArray<AActor*>& Actors
 		{
 			bComponentInside = true;
 			FMeshDescription* MeshDescription = MeshDescriptions[StaticMeshComponent];
-			TVertexAttributesConstRef<FVector> VertexPositions = MeshDescription->VertexAttributes().GetAttributesRef<FVector>(MeshAttribute::Vertex::Position);
+			TVertexAttributesConstRef<FVector3f> VertexPositions = MeshDescription->VertexAttributes().GetAttributesRef<FVector3f>(MeshAttribute::Vertex::Position);
 			for (FVertexID VertexID : MeshDescription->Vertices().GetElementIDs())
 			{
 				if (Volume->QueryDistance(VertexPositions[VertexID]) > MaxDistance)
@@ -586,7 +586,7 @@ void FJacketingProcess::ApplyJacketingOnMeshActors(const TArray<AActor*>& Actors
 		Progress->EnterProgressFrame(VertexTesting, FText::FromString(TEXT("Checking occlusion of vertices ...")));
 	}
 
-	typedef TTuple<UStaticMeshComponent*, const FVector*, bool> FVertexData;
+	typedef TTuple<UStaticMeshComponent*, const FVector3f*, bool> FVertexData;
 	TArray<FVertexData> VertexDataArray;
 	VertexDataArray.Reserve(VertexCount);
 
@@ -600,7 +600,7 @@ void FJacketingProcess::ApplyJacketingOnMeshActors(const TArray<AActor*>& Actors
 			continue;
 		}
 
-		TVertexAttributesConstRef<FVector> VertexPositions = MeshDescriptions[StaticMeshComponent]->VertexAttributes().GetAttributesRef<FVector>(MeshAttribute::Vertex::Position);
+		TVertexAttributesConstRef<FVector3f> VertexPositions = MeshDescriptions[StaticMeshComponent]->VertexAttributes().GetAttributesRef<FVector3f>(MeshAttribute::Vertex::Position);
 		for (FVertexID VertexID : MeshDescriptions[StaticMeshComponent]->Vertices().GetElementIDs())
 		{
 			VertexDataArray.Emplace(FVertexData(StaticMeshComponent, &VertexPositions[VertexID], false));
@@ -708,14 +708,14 @@ void FJacketingProcess::ApplyJacketingOnMeshActors(const TArray<AActor*>& Actors
 			//Apply the inverse component transform
 			FStaticMeshOperations::ApplyTransform(NewMeshDescription, Transform);
 
-			TVertexInstanceAttributesRef<FVector> VertexInstanceNormals = NewMeshDescription.VertexInstanceAttributes().GetAttributesRef<FVector>(MeshAttribute::VertexInstance::Normal);
-			TVertexInstanceAttributesRef<FVector> VertexInstanceTangents = NewMeshDescription.VertexInstanceAttributes().GetAttributesRef<FVector>(MeshAttribute::VertexInstance::Tangent);
+			TVertexInstanceAttributesRef<FVector3f> VertexInstanceNormals = NewMeshDescription.VertexInstanceAttributes().GetAttributesRef<FVector3f>(MeshAttribute::VertexInstance::Normal);
+			TVertexInstanceAttributesRef<FVector3f> VertexInstanceTangents = NewMeshDescription.VertexInstanceAttributes().GetAttributesRef<FVector3f>(MeshAttribute::VertexInstance::Tangent);
 
 			MeshDescription->Empty();
 
 			//Create the missing normals and tangents on polygons because FStaticMeshOperations::ComputeTangentsAndNormals requires it
-			if(!NewMeshDescription.TriangleAttributes().GetAttributesRef<FVector>(MeshAttribute::Triangle::Normal).IsValid()
-				|| !NewMeshDescription.TriangleAttributes().GetAttributesRef<FVector>(MeshAttribute::Triangle::Tangent).IsValid())
+			if(!NewMeshDescription.TriangleAttributes().GetAttributesRef<FVector3f>(MeshAttribute::Triangle::Normal).IsValid()
+				|| !NewMeshDescription.TriangleAttributes().GetAttributesRef<FVector3f>(MeshAttribute::Triangle::Tangent).IsValid())
 			{
 				FStaticMeshOperations::ComputeTriangleTangentsAndNormals(NewMeshDescription);
 			}
diff --git a/Engine/Plugins/Runtime/nDisplay/Source/DisplayClusterShaders/Private/Shaders/DisplayClusterShadersWarpblend_ICVFX.cpp b/Engine/Plugins/Runtime/nDisplay/Source/DisplayClusterShaders/Private/Shaders/DisplayClusterShadersWarpblend_ICVFX.cpp
--- a/Engine/Plugins/Runtime/nDisplay/Source/DisplayClusterShaders/Private/Shaders/DisplayClusterShadersWarpblend_ICVFX.cpp
+++ b/Engine/Plugins/Runtime/nDisplay/Source/DisplayClusterShaders/Private/Shaders/DisplayClusterShadersWarpblend_ICVFX.cpp
@@ -156,7 +156,7 @@ BEGIN_SHADER_PARAMETER_STRUCT(FIcvfxVertexShaderParameters, )
 	SHADER_PARAMETER(FVector4, DrawRectangleInvTargetSizeAndTextureSize)
 	SHADER_PARAMETER(FVector4, DrawRectangleUVScaleBias)
 
-	SHADER_PARAMETER(FMatrix, MeshToStageProjectionMatrix)
+	SHADER_PARAMETER(FMatrix44f, MeshToStageProjectionMatrix)
 END_SHADER_PARAMETER_STRUCT()
 
 BEGIN_SHADER_PARAMETER_STRUCT(FIcvfxPixelShaderParameters, )
@@ -180,9 +180,9 @@ BEGIN_SHADER_PARAMETER_STRUCT(FIcvfxPixelShaderParameters, )
 	SHADER_PARAMETER_SAMPLER(SamplerState, ChromakeyCameraSampler)
 	SHADER_PARAMETER_SAMPLER(SamplerState, ChromakeyMarkerSampler)
 
-	SHADER_PARAMETER(FMatrix, ViewportTextureProjectionMatrix)
-	SHADER_PARAMETER(FMatrix, OverlayProjectionMatrix)
-	SHADER_PARAMETER(FMatrix, InnerCameraProjectionMatrix)
+	SHADER_PARAMETER(FMatrix44f, ViewportTextureProjectionMatrix)
+	SHADER_PARAMETER(FMatrix44f, OverlayProjectionMatrix)
+	SHADER_PARAMETER(FMatrix44f, InnerCameraProjectionMatrix)
 
 	SHADER_PARAMETER(float, AlphaEmbeddedGamma)
 
diff --git a/Engine/Plugins/Runtime/nDisplay/Source/DisplayClusterShaders/Private/Shaders/DisplayClusterShadersWarpblend_MPCDI.cpp b/Engine/Plugins/Runtime/nDisplay/Source/DisplayClusterShaders/Private/Shaders/DisplayClusterShadersWarpblend_MPCDI.cpp
--- a/Engine/Plugins/Runtime/nDisplay/Source/DisplayClusterShaders/Private/Shaders/DisplayClusterShadersWarpblend_MPCDI.cpp
+++ b/Engine/Plugins/Runtime/nDisplay/Source/DisplayClusterShaders/Private/Shaders/DisplayClusterShadersWarpblend_MPCDI.cpp
@@ -85,7 +85,7 @@ BEGIN_SHADER_PARAMETER_STRUCT(FMpcdiVertexShaderParameters, )
 	SHADER_PARAMETER(FVector4, DrawRectangleInvTargetSizeAndTextureSize)
 	SHADER_PARAMETER(FVector4, DrawRectangleUVScaleBias)
 
-	SHADER_PARAMETER(FMatrix, MeshToStageProjectionMatrix)
+	SHADER_PARAMETER(FMatrix44f, MeshToStageProjectionMatrix)
 END_SHADER_PARAMETER_STRUCT()
 
 BEGIN_SHADER_PARAMETER_STRUCT(FMpcdiPixelShaderParameters, )
@@ -99,7 +99,7 @@ BEGIN_SHADER_PARAMETER_STRUCT(FMpcdiPixelShaderParameters, )
 	SHADER_PARAMETER_SAMPLER(SamplerState, AlphaMapSampler)
 	SHADER_PARAMETER_SAMPLER(SamplerState, BetaMapSampler)
 
-	SHADER_PARAMETER(FMatrix, ViewportTextureProjectionMatrix)
+	SHADER_PARAMETER(FMatrix44f, ViewportTextureProjectionMatrix)
 
 	SHADER_PARAMETER(float, AlphaEmbeddedGamma)
 END_SHADER_PARAMETER_STRUCT()
diff --git a/Engine/Plugins/VirtualProduction/RemoteControl/Source/RemoteControl/Private/Tests/RemoteControlPropertyHandleTests.cpp b/Engine/Plugins/VirtualProduction/RemoteControl/Source/RemoteControl/Private/Tests/RemoteControlPropertyHandleTests.cpp
--- a/Engine/Plugins/VirtualProduction/RemoteControl/Source/RemoteControl/Private/Tests/RemoteControlPropertyHandleTests.cpp
+++ b/Engine/Plugins/VirtualProduction/RemoteControl/Source/RemoteControl/Private/Tests/RemoteControlPropertyHandleTests.cpp
@@ -885,20 +885,20 @@ namespace RemoteControlAPIIntegrationTest
 		}
 
 		{
-			float GetFloatValue = 0;
+			FVector::FReal GetFloatValue = 0;
 			ArrayOfVectorsPropertyHandle_O_X->GetValue(GetFloatValue);
 			Test.TestEqual(TEXT("Value in UObject should be the same as a value from Property Handle."), TestVector_0.X, GetFloatValue);
 
-			float SetFloatValue = 0;
+			FVector::FReal SetFloatValue = 0;
 			ArrayOfVectorsPropertyHandle_O_X->SetValue(SetFloatValue);
 			Test.TestEqual(TEXT("Value in UObject should be the same as a value from Property Handle."), TestObject->ArrayOfVectors[0].X, SetFloatValue);
 		}
 		{
-			float GetFloatValue = 0;
+			FVector::FReal GetFloatValue = 0;
 			ArrayOfVectorsPropertyHandle_1_X->GetValue(GetFloatValue);
 			Test.TestEqual(TEXT("Value in UObject should be the same as a value from Property Handle."), TestVector_1.X, GetFloatValue);
 
-			float SetFloatValue = 0;
+			FVector::FReal SetFloatValue = 0;
 			ArrayOfVectorsPropertyHandle_1_X->SetValue(SetFloatValue);
 			Test.TestEqual(TEXT("Value in UObject should be the same as a value from Property Handle."), TestObject->ArrayOfVectors[1].X, SetFloatValue);
 		}
