Dataprep - Filter by Jacketing should't accept negative values

diff --git a/Engine/Plugins/Experimental/Enterprise/DataprepGeometryOperations/Source/Private/DataprepGeometryFilters.h b/Engine/Plugins/Experimental/Enterprise/DataprepGeometryOperations/Source/Private/DataprepGeometryFilters.h
--- a/Engine/Plugins/Experimental/Enterprise/DataprepGeometryOperations/Source/Private/DataprepGeometryFilters.h
+++ b/Engine/Plugins/Experimental/Enterprise/DataprepGeometryOperations/Source/Private/DataprepGeometryFilters.h
@@ -38,10 +38,10 @@ private:
 
 private:
 	/** Accuracy of the distance field approximation, in cm. */
-	UPROPERTY(EditAnywhere, Category = JacketingFilter, meta = (UIMin = "0.1", UIMax = "100"))
+	UPROPERTY(EditAnywhere, Category = JacketingFilter, meta = (UIMin = "0.1", UIMax = "100", ClampMin = "0"))
 	float Accuracy = 3.0f;
 
 	/** Merge distance used to fill gap, in cm. */
-	UPROPERTY(EditAnywhere, Category = JacketingFilter, meta = (UIMin = "0.1", UIMax = "100"))
+	UPROPERTY(EditAnywhere, Category = JacketingFilter, meta = (UIMin = "0.1", UIMax = "100", ClampMin = "0"))
 	float MergeDistance = 4.0f;
 };
