DataprepGeometryOperations: Set a ShortName to reduce the length of intermediate paths

#trivial
#rnx

diff --git a/Engine/Plugins/Experimental/Enterprise/DataprepGeometryOperations/Source/DataprepGeometryOperations.Build.cs b/Engine/Plugins/Experimental/Enterprise/DataprepGeometryOperations/Source/DataprepGeometryOperations.Build.cs
--- a/Engine/Plugins/Experimental/Enterprise/DataprepGeometryOperations/Source/DataprepGeometryOperations.Build.cs
+++ b/Engine/Plugins/Experimental/Enterprise/DataprepGeometryOperations/Source/DataprepGeometryOperations.Build.cs
@@ -6,6 +6,8 @@ namespace UnrealBuildTool.Rules
 	{
 		public DataprepGeometryOperations(ReadOnlyTargetRules Target) : base(Target)
 		{
+			ShortName = "DataprepGO";
+
 			PrivateDependencyModuleNames.AddRange(
 				new string[]
 				{
