Add more consistent clarification to the bDiscardAttributes UPROPERTY in various places where it is exposed w/ comment

diff --git a/Engine/Plugins/Experimental/Enterprise/DataprepGeometryOperations/Source/Private/DataprepGeometryOperations.h b/Engine/Plugins/Experimental/Enterprise/DataprepGeometryOperations/Source/Private/DataprepGeometryOperations.h
--- a/Engine/Plugins/Experimental/Enterprise/DataprepGeometryOperations/Source/Private/DataprepGeometryOperations.h
+++ b/Engine/Plugins/Experimental/Enterprise/DataprepGeometryOperations/Source/Private/DataprepGeometryOperations.h
@@ -38,7 +38,7 @@ public:
 	UPROPERTY(EditAnywhere, Category = Remeshing, meta = (DisplayName = "Smoothing Rate", UIMin = "0.0", UIMax = "1.0", ClampMin = "0.0", ClampMax = "1.0"))
 	float SmoothingStrength = 0.25;
 
-	/** If true, UVs and Normals are discarded  */
+	/** Discard UVs and existing normals, allowing the remesher to ignore any UV and normal seams. New per-vertex normals are computed. */
 	UPROPERTY(EditAnywhere, Category = Remeshing)
 	bool bDiscardAttributes = false;
 
@@ -143,7 +143,7 @@ public:
 	UPROPERTY(EditAnywhere, Category = Options, meta = (UIMin = "0", UIMax = "100"))
 	int TargetPercentage = 50;
 
-	/** If true, UVs and Normals are discarded  */
+	/** Discard UVs and existing normals, allowing the simplifier to ignore any UV and normal seams. New per-vertex normals are computed. */
 	UPROPERTY(EditAnywhere, Category = Options)
 	bool bDiscardAttributes = false;
 
diff --git a/Engine/Plugins/Experimental/GeometryScripting/Source/GeometryScriptingCore/Public/GeometryScript/MeshRemeshFunctions.h b/Engine/Plugins/Experimental/GeometryScripting/Source/GeometryScriptingCore/Public/GeometryScript/MeshRemeshFunctions.h
--- a/Engine/Plugins/Experimental/GeometryScripting/Source/GeometryScriptingCore/Public/GeometryScript/MeshRemeshFunctions.h
+++ b/Engine/Plugins/Experimental/GeometryScripting/Source/GeometryScriptingCore/Public/GeometryScript/MeshRemeshFunctions.h
@@ -56,7 +56,7 @@ struct GEOMETRYSCRIPTINGCORE_API FGeometryScriptRemeshOptions
 	GENERATED_BODY()
 public:
 
-	/** When enabled, all mesh attributes are discarded, so UV and Normal Seams can be freely */
+	/** When enabled, all mesh attributes are discarded, so UV and Normal Seams can be ignored. New per-vertex normals are computed. */
 	UPROPERTY(BlueprintReadWrite, Category = Options)
 	bool bDiscardAttributes = false;
 
diff --git a/Engine/Plugins/Experimental/HairModelingToolset/Source/HairModelingToolset/Public/GenerateLODMeshesTool.h b/Engine/Plugins/Experimental/HairModelingToolset/Source/HairModelingToolset/Public/GenerateLODMeshesTool.h
--- a/Engine/Plugins/Experimental/HairModelingToolset/Source/HairModelingToolset/Public/GenerateLODMeshesTool.h
+++ b/Engine/Plugins/Experimental/HairModelingToolset/Source/HairModelingToolset/Public/GenerateLODMeshesTool.h
@@ -103,7 +103,7 @@ public:
 	UPROPERTY()
 	int TargetCount = 1000;
 
-	/** If true, UVs and Normals are discarded  */
+	/** Discard UVs and existing normals, allowing the simplifier to ignore any UV and normal seams. New per-vertex normals are computed. */
 	//UPROPERTY(EditAnywhere, Category = Options)
 	UPROPERTY()
 	bool bDiscardAttributes = false;
diff --git a/Engine/Plugins/Experimental/MeshModelingToolsetExp/Source/MeshModelingToolsEditorOnlyExp/Public/SimplifyMeshTool.h b/Engine/Plugins/Experimental/MeshModelingToolsetExp/Source/MeshModelingToolsEditorOnlyExp/Public/SimplifyMeshTool.h
--- a/Engine/Plugins/Experimental/MeshModelingToolsetExp/Source/MeshModelingToolsEditorOnlyExp/Public/SimplifyMeshTool.h
+++ b/Engine/Plugins/Experimental/MeshModelingToolsetExp/Source/MeshModelingToolsEditorOnlyExp/Public/SimplifyMeshTool.h
@@ -76,7 +76,7 @@ public:
 	UPROPERTY(EditAnywhere, Category = Options, meta = (UIMin = "0.001", ClampMin = "0.0", UIMax = "90.0", ClampMax = "180.0", EditCondition = "SimplifierType == ESimplifyType::MinimalPolygroup"))
 	float PolyEdgeAngleTolerance = .1;
 
-	/** If true, UVs and Normals are discarded  */
+	/** Discard UVs and existing normals, allowing the simplifier to ignore any UV and normal seams. New per-vertex normals are computed. */
 	UPROPERTY(EditAnywhere, Category = Options)
 	bool bDiscardAttributes;
 
diff --git a/Engine/Plugins/Experimental/MeshModelingToolsetExp/Source/MeshModelingToolsExp/Public/RemeshMeshTool.h b/Engine/Plugins/Experimental/MeshModelingToolsetExp/Source/MeshModelingToolsExp/Public/RemeshMeshTool.h
--- a/Engine/Plugins/Experimental/MeshModelingToolsetExp/Source/MeshModelingToolsExp/Public/RemeshMeshTool.h
+++ b/Engine/Plugins/Experimental/MeshModelingToolsetExp/Source/MeshModelingToolsExp/Public/RemeshMeshTool.h
@@ -52,7 +52,7 @@ public:
 	UPROPERTY(EditAnywhere, Category = Remeshing)
 	ERemeshSmoothingType SmoothingType;
 
-	/** If true, UVs and existing normals are discarded, allowing the remesher to ignore any UV and normal seams. New per-vertex normals are computed. */
+	/** Discard UVs and existing normals, allowing the remesher to ignore any UV and normal seams. New per-vertex normals are computed. */
 	UPROPERTY(EditAnywhere, Category = Remeshing)
 	bool bDiscardAttributes;
 
