update another usage of FDynamicMeshEditor::SplitMesh to use the TUniquePtr version; MeshPartitionResectionTool and DataprepGeometryOperations

#Synced-CL 50522370

diff --git a/Engine/Plugins/Experimental/Enterprise/DataprepGeometryOperations/Source/Private/DataprepGeometryOperations.cpp b/Engine/Plugins/Experimental/Enterprise/DataprepGeometryOperations/Source/Private/DataprepGeometryOperations.cpp
--- a/Engine/Plugins/Experimental/Enterprise/DataprepGeometryOperations/Source/Private/DataprepGeometryOperations.cpp
+++ b/Engine/Plugins/Experimental/Enterprise/DataprepGeometryOperations/Source/Private/DataprepGeometryOperations.cpp
@@ -631,7 +631,7 @@ void UDataprepPlaneCutOperation::PerformCutting(
 		Results[Index] = CutStaticMesh(InCutPlaneTransforms[Index], InStaticMeshes[Index]);
 	});
 
-	TArray<TArray<FDynamicMesh3>> AllSplitMeshes;
+	TArray<TArray<TUniquePtr<FDynamicMesh3>>> AllSplitMeshes;
 	AllSplitMeshes.SetNum(InStaticMeshes.Num());
 
 	for (int OrigMeshIdx = 0; OrigMeshIdx < InStaticMeshes.Num(); OrigMeshIdx++)
@@ -662,8 +662,8 @@ void UDataprepPlaneCutOperation::PerformCutting(
 			UE::Geometry::TDynamicMeshScalarTriangleAttribute<int>* SubMeshIDs =
 				static_cast<UE::Geometry::TDynamicMeshScalarTriangleAttribute<int>*>(UseMesh->Attributes()->GetAttachedAttribute(
 					UE::Geometry::FPlaneCutOp::ObjectIndexAttribute));
-			TArray<UE::Geometry::FDynamicMesh3>& SplitMeshes = AllSplitMeshes[OrigMeshIdx];
-			bool bWasSplit = UE::Geometry::FDynamicMeshEditor::SplitMesh(UseMesh, SplitMeshes, [SubMeshIDs](int TID)
+			TArray<TUniquePtr<UE::Geometry::FDynamicMesh3>>& SplitMeshes = AllSplitMeshes[OrigMeshIdx];
+			bool bWasSplit = UE::Geometry::FDynamicMeshEditor::SplitMesh(UseMesh, SplitMeshes, false, [SubMeshIDs](int TID)
 			{
 				return SubMeshIDs->GetValue(TID);
 			});
@@ -693,7 +693,7 @@ void UDataprepPlaneCutOperation::PerformCutting(
 						{
 							FTransform Transform = ComponentTarget->GetComponentTransform();
 
-							FDynamicMesh3* Mesh = &SplitMeshes[AddMeshIdx];
+							FDynamicMesh3* Mesh = SplitMeshes[AddMeshIdx].Get();
 
 							TUniquePtr<FMeshDescription> MeshDescription = MakeUnique<FMeshDescription>();
 							FStaticMeshAttributes Attributes(*MeshDescription);
@@ -763,7 +763,7 @@ void UDataprepPlaneCutOperation::PerformCutting(
 						}
 					}
 				}
-				UseMesh = &SplitMeshes[0];
+				UseMesh = SplitMeshes[0].Get();
 			}
 		}
 
